Move TranslationPicker from Editor to Developer, to run in game

diff --git a/Engine/Source/Developer/TranslationEditor/Private/CustomFontColumn.cpp b/Engine/Source/Developer/TranslationEditor/Private/CustomFontColumn.cpp
new file mode 100644
--- /dev/null
+++ b/Engine/Source/Developer/TranslationEditor/Private/CustomFontColumn.cpp
@@ -0,0 +1,115 @@
+// Copyright Epic Games, Inc. All Rights Reserved.
+#if WITH_EDITOR
+
+#include "CustomFontColumn.h"
+
+#include "HAL/PlatformMath.h"
+#include "IPropertyTable.h"
+#include "IPropertyTableCell.h"
+#include "IPropertyTableColumn.h"
+#include "Internationalization/Internationalization.h"
+#include "Internationalization/Text.h"
+#include "Misc/Attribute.h"
+#include "Misc/Optional.h"
+#include "Modules/ModuleManager.h"
+#include "PropertyEditorModule.h"
+#include "PropertyPath.h"
+#include "SlotBase.h"
+#include "Templates/TypeHash.h"
+#include "Types/SlateEnums.h"
+#include "UObject/UnrealType.h"
+#include "UObject/WeakFieldPtr.h"
+#include "Widgets/DeclarativeSyntaxSupport.h"
+#include "Widgets/Input/SButton.h"
+#include "Widgets/Input/SSpinBox.h"
+#include "Widgets/SBoxPanel.h"
+#include "Widgets/Text/STextBlock.h"
+
+class IPropertyHandle;
+class IPropertyTableUtilities;
+class SWidget;
+
+#define LOCTEXT_NAMESPACE "PropertyTable.CustomFontColumn"
+
+bool FCustomFontColumn::Supports( const TSharedRef< IPropertyTableColumn >& Column, const TSharedRef< IPropertyTableUtilities >& Utilities ) const
+{
+	bool IsSupported = false;
+	
+	if( Column->GetDataSource()->IsValid() )
+	{
+		TSharedPtr< FPropertyPath > PropertyPath = Column->GetDataSource()->AsPropertyPath();
+		if( PropertyPath.IsValid() && PropertyPath->GetNumProperties() > 0 )
+		{
+			const FPropertyInfo& PropertyInfo = PropertyPath->GetRootProperty();
+			FProperty* Property = PropertyInfo.Property.Get();
+			if (SupportedProperties.Contains(Property))
+			{
+				IsSupported = true;
+			}
+		}
+	}
+
+	return IsSupported;
+}
+
+TSharedPtr< SWidget > FCustomFontColumn::CreateColumnLabel( const TSharedRef< IPropertyTableColumn >& Column, const TSharedRef< IPropertyTableUtilities >& Utilities, const FName& Style ) const
+{
+	return 
+	SNew( SHorizontalBox )
+	+SHorizontalBox::Slot()
+	.AutoWidth()
+	.HAlign( HAlign_Left )
+	.VAlign( VAlign_Center )
+	.FillWidth(2)
+	[
+		SNew( STextBlock )
+		.Font( FAppStyle::GetFontStyle( Style ) )
+		.Text( Column->GetDisplayName() )
+	]
+	+SHorizontalBox::Slot()
+	.AutoWidth()
+	.HAlign( HAlign_Center )
+	.VAlign( VAlign_Center )
+	.FillWidth(2)
+	[
+		SNew( SButton )
+		.HAlign( HAlign_Center )
+		.VAlign( VAlign_Center )
+		.OnClicked(OnChangeFontButtonClicked)
+		[
+				SNew( STextBlock )
+				.Text((FText::Format(LOCTEXT("ChooseFontButton", "Choose {0} Font: "), Column->GetDisplayName())))
+		]
+	]
+	+SHorizontalBox::Slot()
+	.AutoWidth()
+	.HAlign( HAlign_Center )
+	.VAlign( VAlign_Center )
+	.FillWidth(1)
+	.Padding( 0.0f, 4.0f, 0.0f, 4.0f )
+	[
+		SNew( SSpinBox<int32> )
+		.Delta(1)
+		.MinValue(1)
+		.MaxValue(100)
+		.Value(Font.Size)
+		.OnValueCommitted( OnFontSizeValueCommitted )
+	]
+	;
+}
+
+TSharedPtr< IPropertyTableCellPresenter > FCustomFontColumn::CreateCellPresenter( const TSharedRef< IPropertyTableCell >& Cell, const TSharedRef< IPropertyTableUtilities >& Utilities, const FName& Style ) const
+{
+	TSharedPtr< IPropertyHandle > PropertyHandle = Cell->GetPropertyHandle();
+	if( PropertyHandle.IsValid() )
+	{
+		FPropertyEditorModule& PropertyEditorModule = FModuleManager::GetModuleChecked<FPropertyEditorModule>( "PropertyEditor" );
+		return PropertyEditorModule.CreateTextPropertyCellPresenter(Cell->GetNode().ToSharedRef(), Utilities, &Font);
+	}
+
+	return NULL;
+}
+
+#undef LOCTEXT_NAMESPACE
+
+#endif // WITH_EDITOR
diff --git a/Engine/Source/Developer/TranslationEditor/Private/CustomFontColumn.h b/Engine/Source/Developer/TranslationEditor/Private/CustomFontColumn.h
new file mode 100644
--- /dev/null
+++ b/Engine/Source/Developer/TranslationEditor/Private/CustomFontColumn.h
@@ -0,0 +1,76 @@
+// Copyright Epic Games, Inc. All Rights Reserved.
+#pragma once
+#if WITH_EDITOR
+
+#include "Containers/Array.h"
+#include "UserInterface/PropertyTable/PropertyTableConstants.h"
+#include "Fonts/SlateFontInfo.h"
+#include "Framework/SlateDelegates.h"
+#include "HAL/PlatformCrt.h"
+#include "IPropertyTableCustomColumn.h"
+#include "Styling/AppStyle.h"
+#include "Templates/SharedPointer.h"
+#include "UObject/NameTypes.h"
+
+class FProperty;
+class IPropertyTableCell;
+class IPropertyTableCellPresenter;
+class IPropertyTableColumn;
+class IPropertyTableUtilities;
+class SWidget;
+
+/**
+ * A property table custom column used to display text cells in a chosen font
+ */
+class FCustomFontColumn : public IPropertyTableCustomColumn
+{
+public:
+	FCustomFontColumn(FSlateFontInfo InFont = FAppStyle::GetFontStyle( PropertyTableConstants::NormalFontStyle ), FOnClicked InOnChangeFontButtonClicked = NULL, FOnInt32ValueCommitted InOnFontSizeValueCommitted = NULL)
+		:	Font(InFont)
+		,	OnChangeFontButtonClicked(InOnChangeFontButtonClicked)
+		,	OnFontSizeValueCommitted(InOnFontSizeValueCommitted)
+	{}
+
+	/** Begin IPropertyTableCustomColumn interface */
+	virtual bool Supports( const TSharedRef< IPropertyTableColumn >& Column, const TSharedRef< IPropertyTableUtilities >& Utilities ) const override;
+	virtual TSharedPtr< SWidget > CreateColumnLabel( const TSharedRef< IPropertyTableColumn >& Column, const TSharedRef< IPropertyTableUtilities >& Utilities, const FName& Style ) const override;
+	virtual TSharedPtr< IPropertyTableCellPresenter > CreateCellPresenter( const TSharedRef< IPropertyTableCell >& Cell, const TSharedRef< IPropertyTableUtilities >& Utilities, const FName& Style ) const override;
+	/** End IPropertyTableCustomColumn interface */
+
+private:
+
+	/** Font to use for this column's cells */
+	FSlateFontInfo Font;
+
+	/** List of properties that this Custom Font Column should be used to display */
+	TArray<FProperty*> SupportedProperties;
+
+	/** Function to call when Change Font button clicked */
+	FOnClicked OnChangeFontButtonClicked;
+
+	/** Function to call when Font Size SpinBox's Value is committed */
+	FOnInt32ValueCommitted OnFontSizeValueCommitted;
+
+public:
+	void SetFont(FSlateFontInfo InFont)
+	{
+		Font = InFont;
+	}
+
+	void AddSupportedProperty(FProperty* Property)
+	{
+		SupportedProperties.Add(Property);
+	}
+
+	void SetOnChangeFontButtonClicked(FOnClicked InOnClicked)
+	{
+		OnChangeFontButtonClicked = InOnClicked;
+	}
+
+	void SetOnFontSizeValueCommitted(FOnInt32ValueCommitted InOnFontSizeValueCommitted)
+	{
+		OnFontSizeValueCommitted = InOnFontSizeValueCommitted;
+	}
+};
+
+#endif // WITH_EDITOR
diff --git a/Engine/Source/Developer/TranslationEditor/Private/ITranslationEditor.cpp b/Engine/Source/Developer/TranslationEditor/Private/ITranslationEditor.cpp
new file mode 100644
--- /dev/null
+++ b/Engine/Source/Developer/TranslationEditor/Private/ITranslationEditor.cpp
@@ -0,0 +1,94 @@
+// Copyright Epic Games, Inc. All Rights Reserved.
+
+#include "ITranslationEditor.h"
+#include "TranslationPickerWidget.h"
+#if WITH_EDITOR
+#include "Modules/ModuleManager.h"
+#include "TranslationEditorModule.h"
+#include "LocalizationConfigurationScript.h"
+
+/** To keep track of what translation editors are open editing which archive files */
+TMap<FString, ITranslationEditor*> ITranslationEditor::OpenTranslationEditors = TMap<FString, ITranslationEditor*>();
+
+void ITranslationEditor::OpenTranslationEditor(const FString& InManifestFile, const FString& InNativeArchiveFile, const FString& InArchiveFileToEdit)
+{
+	// Only open a new translation editor if there isn't another one open for this archive file already
+	if (!OpenTranslationEditors.Contains(InArchiveFileToEdit))
+	{
+		FTranslationEditorModule& TranslationEditorModule = FModuleManager::LoadModuleChecked<FTranslationEditorModule>("TranslationEditor");
+		bool bLoadedSuccessfully = false;
+		ITranslationEditor* NewTranslationEditor = (ITranslationEditor*)&(TranslationEditorModule.CreateTranslationEditor(InManifestFile, InNativeArchiveFile, InArchiveFileToEdit, bLoadedSuccessfully).Get());
+
+		NewTranslationEditor->RegisterTranslationEditor();
+
+		if (!bLoadedSuccessfully)
+		{
+			NewTranslationEditor->CloseWindow(EAssetEditorCloseReason::AssetUnloadingOrInvalid);
+		}
+	}
+	else
+	{
+		// If already editing this archive file, flash the tab that contains the editor that has that file open
+		ITranslationEditor* OpenEditor = *OpenTranslationEditors.Find(InArchiveFileToEdit);
+		OpenEditor->FocusWindow();
+	}
+}
+
+void ITranslationEditor::OpenTranslationEditor(ULocalizationTarget* const LocalizationTarget, const FString& CultureToEdit)
+{
+	check(LocalizationTarget);
+	const FString ManifestFile = LocalizationConfigurationScript::GetManifestPath(LocalizationTarget);
+	FString NativeCultureName;
+	if (LocalizationTarget->Settings.SupportedCulturesStatistics.IsValidIndex(LocalizationTarget->Settings.NativeCultureIndex))
+	{
+		NativeCultureName = LocalizationTarget->Settings.SupportedCulturesStatistics[LocalizationTarget->Settings.NativeCultureIndex].CultureName;
+	}
+	const FString NativeArchiveFile = NativeCultureName.IsEmpty() ? FString() : LocalizationConfigurationScript::GetArchivePath(LocalizationTarget, NativeCultureName);
+	const FString ArchiveFile = LocalizationConfigurationScript::GetArchivePath(LocalizationTarget, CultureToEdit);
+
+	// Only open a new translation editor if there isn't another one open for this archive file already
+	if (!OpenTranslationEditors.Contains(ArchiveFile))
+	{
+		FTranslationEditorModule& TranslationEditorModule = FModuleManager::LoadModuleChecked<FTranslationEditorModule>("TranslationEditor");
+		bool bLoadedSuccessfully = false;
+		ITranslationEditor* NewTranslationEditor = (ITranslationEditor*)&(TranslationEditorModule.CreateTranslationEditor(LocalizationTarget, CultureToEdit, bLoadedSuccessfully).Get());
+
+		NewTranslationEditor->RegisterTranslationEditor();
+
+		if (!bLoadedSuccessfully)
+		{
+			NewTranslationEditor->CloseWindow(EAssetEditorCloseReason::AssetUnloadingOrInvalid);
+		}
+	}
+	else
+	{
+		// If already editing this archive file, flash the tab that contains the editor that has that file open
+		ITranslationEditor* OpenEditor = *OpenTranslationEditors.Find(ArchiveFile);
+		OpenEditor->FocusWindow();
+	}
+}
+
+void ITranslationEditor::RegisterTranslationEditor()
+{
+	ITranslationEditor::OpenTranslationEditors.Add(ArchiveFilePath, this);
+}
+
+void ITranslationEditor::UnregisterTranslationEditor()
+{
+	OpenTranslationEditors.Remove(ArchiveFilePath);
+}
+
+void ITranslationEditor::OnClose()
+{
+	ITranslationEditor::UnregisterTranslationEditor();
+}
+
+#endif // WITH_EDITOR
+
+void ITranslationEditor::OpenTranslationPicker()
+{
+	if (!TranslationPickerManager::IsPickerWindowOpen())
+	{
+		TranslationPickerManager::OpenPickerWindow();
+	}
+}
diff --git a/Engine/Source/Developer/TranslationEditor/Private/TranslationDataManager.cpp b/Engine/Source/Developer/TranslationEditor/Private/TranslationDataManager.cpp
new file mode 100644
--- /dev/null
+++ b/Engine/Source/Developer/TranslationEditor/Private/TranslationDataManager.cpp
@@ -0,0 +1,1038 @@
+// Copyright Epic Games, Inc. All Rights Reserved.
+
+#include "TranslationDataManager.h"
+#include "PortableObjectPipeline.h"
+#include "Internationalization/InternationalizationManifest.h"
+#include "Internationalization/InternationalizationArchive.h"
+#include "Misc/FileHelper.h"
+#include "Misc/Paths.h"
+#include "Misc/FeedbackContext.h"
+#include "Misc/App.h"
+#include "Dom/JsonObject.h"
+#include "Serialization/JsonSerializer.h"
+#include "Styling/AppStyle.h"
+#include "SourceControlOperations.h"
+#include "ISourceControlState.h"
+#include "ISourceControlProvider.h"
+#include "ISourceControlModule.h"
+#include "SourceControlHelpers.h"
+#include "TranslationUnit.h"
+#include "Logging/MessageLog.h"
+#include "TextLocalizationResourceGenerator.h"
+#include "Framework/Notifications/NotificationManager.h"
+#include "Widgets/Notifications/SNotificationList.h"
+#include "Internationalization/Culture.h"
+#include "PortableObjectFormatDOM.h"
+#include "ILocalizationServiceModule.h"
+#include "LocalizationModule.h"
+#include "LocalizationTargetTypes.h"
+#include "LocalizationConfigurationScript.h"
+#include "Serialization/JsonInternationalizationArchiveSerializer.h"
+#include "Serialization/JsonInternationalizationManifestSerializer.h"
+#include "GeneralProjectSettings.h"
+
+
+DEFINE_LOG_CATEGORY_STATIC(LogTranslationEditor, Log, All);
+
+#define LOCTEXT_NAMESPACE "TranslationDataManager"
+
+struct FLocTextIdentity
+{
+public:
+	FLocTextIdentity(FLocKey InNamespace, FLocKey InKey)
+		: Namespace(MoveTemp(InNamespace))
+		, Key(MoveTemp(InKey))
+	{
+	}
+
+	FORCEINLINE const FLocKey& GetNamespace() const
+	{
+		return Namespace;
+	}
+
+	FORCEINLINE const FLocKey& GetKey() const
+	{
+		return Key;
+	}
+
+	FORCEINLINE bool operator==(const FLocTextIdentity& Other) const
+	{
+		return Namespace == Other.Namespace && Key == Other.Key;
+	}
+
+	FORCEINLINE bool operator!=(const FLocTextIdentity& Other) const
+	{
+		return Namespace != Other.Namespace || Key != Other.Key;
+	}
+
+	friend inline uint32 GetTypeHash(const FLocTextIdentity& Id)
+	{
+		return HashCombine(GetTypeHash(Id.Namespace), GetTypeHash(Id.Key));
+	}
+
+private:
+	FLocKey Namespace;
+	FLocKey Key;
+};
+
+FTranslationDataManager::FTranslationDataManager( const FString& InManifestFilePath, const FString& InNativeArchiveFilePath, const FString& InArchiveFilePath )
+	: OpenedManifestFilePath(InManifestFilePath)
+	, NativeArchiveFilePath(InNativeArchiveFilePath)
+	, OpenedArchiveFilePath(InArchiveFilePath)
+	, bLoadedSuccessfully(true)
+{
+	Initialize();
+}
+
+FTranslationDataManager::FTranslationDataManager(ULocalizationTarget* const LocalizationTarget, const FString& CultureToEdit)
+	: bLoadedSuccessfully(true)
+{
+	check(LocalizationTarget);
+ 
+	const FString ManifestFile = LocalizationConfigurationScript::GetManifestPath(LocalizationTarget);
+	FString NativeCultureName;
+	if (LocalizationTarget->Settings.SupportedCulturesStatistics.IsValidIndex(LocalizationTarget->Settings.NativeCultureIndex))
+	{
+		NativeCultureName = LocalizationTarget->Settings.SupportedCulturesStatistics[LocalizationTarget->Settings.NativeCultureIndex].CultureName;
+	}
+	const FString NativeArchiveFile = NativeCultureName.IsEmpty() ? FString() : LocalizationConfigurationScript::GetArchivePath(LocalizationTarget, NativeCultureName);
+	const FString ArchiveFileToEdit = LocalizationConfigurationScript::GetArchivePath(LocalizationTarget, CultureToEdit);
+
+	OpenedManifestFilePath = ManifestFile;
+	NativeArchiveFilePath = NativeArchiveFile;
+	OpenedArchiveFilePath = ArchiveFileToEdit;
+
+	Initialize();
+}
+
+void FTranslationDataManager::Initialize()
+{
+	GWarn->BeginSlowTask(LOCTEXT("LoadingTranslationData", "Loading Translation Data..."), true);
+	TArray<UTranslationUnit*> TranslationUnits;
+
+	ManifestAtHeadRevisionPtr = ReadManifest( OpenedManifestFilePath );
+	if (ManifestAtHeadRevisionPtr.IsValid())
+	{
+		TSharedRef< FInternationalizationManifest > ManifestAtHeadRevision = ManifestAtHeadRevisionPtr.ToSharedRef();
+		int32 ManifestEntriesCount = ManifestAtHeadRevision->GetNumEntriesBySourceText();
+
+		if (ManifestEntriesCount < 1)
+		{
+			bLoadedSuccessfully = false;
+			FFormatNamedArguments Arguments;
+			Arguments.Add( TEXT("ManifestFilePath"), FText::FromString(OpenedManifestFilePath) );
+			Arguments.Add( TEXT("ManifestEntriesCount"), FText::AsNumber(ManifestEntriesCount) );
+			FMessageLog TranslationEditorMessageLog("TranslationEditor");
+			TranslationEditorMessageLog.Error(FText::Format(LOCTEXT("CurrentManifestEmpty", "Most current translation manifest ({ManifestFilePath}) has {ManifestEntriesCount} entries."), Arguments));
+			TranslationEditorMessageLog.Notify(LOCTEXT("TranslationLoadError", "Error Loading Translations!"));
+			TranslationEditorMessageLog.Open(EMessageSeverity::Error);
+		}
+
+		ArchivePtr = ReadArchive(OpenedArchiveFilePath);
+		NativeArchivePtr = NativeArchiveFilePath != OpenedArchiveFilePath ? ReadArchive(NativeArchiveFilePath) : ArchivePtr;
+
+		if (ArchivePtr.IsValid())
+		{
+			int32 NumManifestEntriesParsed = 0;
+
+			GWarn->BeginSlowTask(LOCTEXT("LoadingCurrentManifest", "Loading Entries from Current Translation Manifest..."), true);
+
+			// Get all manifest entries by source text...
+			for (auto ManifestItr = ManifestAtHeadRevision->GetEntriesBySourceTextIterator(); ManifestItr; ++ManifestItr, ++NumManifestEntriesParsed)
+			{
+				GWarn->StatusUpdate(NumManifestEntriesParsed, ManifestEntriesCount, FText::Format(LOCTEXT("LoadingCurrentManifestEntries", "Loading Entry {0} of {1} from Current Translation Manifest..."), FText::AsNumber(NumManifestEntriesParsed), FText::AsNumber(ManifestEntriesCount)));
+				const TSharedRef<FManifestEntry> ManifestEntry = ManifestItr.Value();
+				TMap< FLocTextIdentity, UTranslationUnit* > IdentityToTranslationUnitMap;
+
+				for(auto ContextIter( ManifestEntry->Contexts.CreateConstIterator() ); ContextIter; ++ContextIter)
+				{
+					FTranslationContextInfo ContextInfo;
+					const FManifestContext& AContext = *ContextIter;
+
+					ContextInfo.Context = AContext.SourceLocation;
+					ContextInfo.Key = AContext.Key.GetString();
+
+					// Make sure we have a unique translation unit for each unique identity.
+					UTranslationUnit*& TranslationUnit = IdentityToTranslationUnitMap.FindOrAdd(FLocTextIdentity(ManifestEntry->Namespace, AContext.Key));
+					if (!TranslationUnit)
+					{
+						TranslationUnit = NewObject<UTranslationUnit>();
+						check(TranslationUnit != nullptr);
+						// We want Undo/Redo support
+						TranslationUnit->SetFlags(RF_Transactional);
+						TranslationUnit->HasBeenReviewed = false;
+						TranslationUnit->Source = ManifestEntry->Source.Text;
+						TranslationUnit->Namespace = ManifestEntry->Namespace.GetString();
+						TranslationUnit->Key = AContext.Key.GetString();
+						TranslationUnit->KeyMetaDataObject = AContext.KeyMetadataObj;
+					}
+					
+					if (NativeArchivePtr.IsValid() && NativeArchivePtr != ArchivePtr)
+					{
+						const TSharedPtr<FArchiveEntry> NativeArchiveEntry = NativeArchivePtr->FindEntryByKey(ManifestEntry->Namespace, AContext.Key, AContext.KeyMetadataObj);
+						// If the native archive contains a translation for the source string that isn't identical to the source string, use the translation as the source string.
+						if (NativeArchiveEntry.IsValid() && !NativeArchiveEntry->Translation.IsExactMatch(NativeArchiveEntry->Source))
+						{
+							TranslationUnit->Source = NativeArchiveEntry->Translation.Text;
+						}
+					}
+
+					TranslationUnit->Contexts.Add(ContextInfo);
+				}
+
+
+				TArray<UTranslationUnit*> TranslationUnitsToAdd;
+				IdentityToTranslationUnitMap.GenerateValueArray(TranslationUnitsToAdd);
+				TranslationUnits.Append(TranslationUnitsToAdd);
+			}
+			GWarn->EndSlowTask();
+
+			LoadFromArchive(TranslationUnits);
+		}
+		else  // ArchivePtr.IsValid() is false
+		{
+			bLoadedSuccessfully = false;
+			FFormatNamedArguments Arguments;
+			Arguments.Add( TEXT("ArchiveFilePath"), FText::FromString(OpenedArchiveFilePath) );
+			FMessageLog TranslationEditorMessageLog("TranslationEditor");
+			TranslationEditorMessageLog.Error(FText::Format(LOCTEXT("FailedToLoadCurrentArchive", "Failed to load most current translation archive ({ArchiveFilePath}), unable to load translations."), Arguments));
+			TranslationEditorMessageLog.Notify(LOCTEXT("TranslationLoadError", "Error Loading Translations!"));
+			TranslationEditorMessageLog.Open(EMessageSeverity::Error);
+		}
+	}
+	else  // ManifestAtHeadRevisionPtr.IsValid() is false
+	{
+		bLoadedSuccessfully = false;
+		FFormatNamedArguments Arguments;
+		Arguments.Add( TEXT("ManifestFilePath"), FText::FromString(OpenedManifestFilePath) );
+		FMessageLog TranslationEditorMessageLog("TranslationEditor");
+		TranslationEditorMessageLog.Error(FText::Format(LOCTEXT("FailedToLoadCurrentManifest", "Failed to load most current translation manifest ({ManifestFilePath}), unable to load translations."), Arguments));
+		TranslationEditorMessageLog.Notify(LOCTEXT("TranslationLoadError", "Error Loading Translations!"));
+		TranslationEditorMessageLog.Open(EMessageSeverity::Error);
+	}
+
+	GWarn->EndSlowTask();
+}
+
+FTranslationDataManager::~FTranslationDataManager()
+{
+	RemoveTranslationUnitArrayfromRoot(AllTranslations); // Re-enable garbage collection for all current UTranslationDataObjects
+}
+
+TSharedPtr< FInternationalizationManifest > FTranslationDataManager::ReadManifest( const FString& ManifestFilePathToRead )
+{
+	TSharedPtr<FInternationalizationManifest> InternationalizationManifest = MakeShareable(new FInternationalizationManifest);
+
+	if (!FJsonInternationalizationManifestSerializer::DeserializeManifestFromFile(ManifestFilePathToRead, InternationalizationManifest.ToSharedRef()))
+	{
+		UE_LOG(LogTranslationEditor, Error, TEXT("Could not read manifest file %s."), *ManifestFilePathToRead);
+		InternationalizationManifest.Reset();
+	}
+
+	return InternationalizationManifest;
+}
+
+TSharedPtr< FInternationalizationArchive > FTranslationDataManager::ReadArchive(const FString& ArchiveFilePath)
+{
+	TSharedPtr<FInternationalizationArchive> InternationalizationArchive = MakeShareable(new FInternationalizationArchive);
+
+	if (!FJsonInternationalizationArchiveSerializer::DeserializeArchiveFromFile(ArchiveFilePath, InternationalizationArchive.ToSharedRef(), ManifestAtHeadRevisionPtr, nullptr))
+	{
+		UE_LOG(LogTranslationEditor, Error, TEXT("Could not read archive file %s."), *ArchiveFilePath);
+		InternationalizationArchive.Reset();
+	}
+
+	return InternationalizationArchive;
+}
+
+bool FTranslationDataManager::WriteTranslationData(bool bForceWrite /*= false*/)
+{
+	bool bSuccess = false;
+
+	// If the archive hasn't been loaded correctly, don't try and write anything
+	if (ArchivePtr.IsValid())
+	{
+		TSharedRef< FInternationalizationArchive > Archive = ArchivePtr.ToSharedRef();
+
+		bool bNeedsWrite = false;
+
+		auto WriteTranslationUnits = [&Archive , &bNeedsWrite](const TArray<UTranslationUnit*> InTranslationUnits)
+		{
+			for (UTranslationUnit* TranslationUnit : InTranslationUnits)
+			{
+				if (TranslationUnit)
+				{
+					TSharedPtr<FArchiveEntry> OldArchiveEntry = Archive->FindEntryByKey(TranslationUnit->Namespace, TranslationUnit->Key, TranslationUnit->KeyMetaDataObject);
+					if (TranslationUnit->HasBeenReviewed && (!OldArchiveEntry || !OldArchiveEntry->Source.Text.Equals(TranslationUnit->Source, ESearchCase::CaseSensitive) || !OldArchiveEntry->Translation.Text.Equals(TranslationUnit->Translation, ESearchCase::CaseSensitive)))
+					{
+						Archive->SetTranslation(TranslationUnit->Namespace, TranslationUnit->Key, FLocItem(TranslationUnit->Source), FLocItem(TranslationUnit->Translation), TranslationUnit->KeyMetaDataObject);
+						bNeedsWrite = true;
+					}
+				}
+			}
+		};
+
+		WriteTranslationUnits(Untranslated);
+		WriteTranslationUnits(Review);
+		WriteTranslationUnits(Complete);
+
+		bSuccess = true;
+
+		if (bForceWrite || bNeedsWrite)
+		{
+			TSharedRef<FJsonObject> FinalArchiveJsonObj = MakeShareable(new FJsonObject);
+			FJsonInternationalizationArchiveSerializer::SerializeArchive(Archive, FinalArchiveJsonObj);
+
+			bSuccess = WriteJSONToTextFile(FinalArchiveJsonObj, OpenedArchiveFilePath);
+		}
+	}
+
+	return bSuccess;
+}
+
+bool FTranslationDataManager::WriteJSONToTextFile(TSharedRef<FJsonObject>& Output, const FString& Filename)
+{
+	bool CheckoutAndSaveWasSuccessful = true;
+	bool bPreviouslyCheckedOut = false;
+
+	// If the user specified a reference file - write the entries read from code to a ref file
+	if ( !Filename.IsEmpty() )
+	{
+		// If source control is enabled, try to check out the file. Otherwise just try to write it
+		if (ISourceControlModule::Get().IsEnabled())
+		{
+			// Already checked out?
+			if (CheckedOutFiles.Contains(Filename))
+			{
+				bPreviouslyCheckedOut = true;
+			}
+			else if (!USourceControlHelpers::CheckOutOrAddFile(Filename))
+			{
+				FFormatNamedArguments Arguments;
+				Arguments.Add(TEXT("Filename"), FText::FromString(Filename));
+				// Use Source Control Message Log here because there might be other useful information in that log for the user.
+				FMessageLog SourceControlMessageLog("SourceControl");
+				SourceControlMessageLog.Error(FText::Format(LOCTEXT("CheckoutFailed", "Check out of file '{Filename}' failed."), Arguments));
+				SourceControlMessageLog.Notify(LOCTEXT("TranslationArchiveCheckoutFailed", "Failed to Check Out Translation Archive!"));
+				SourceControlMessageLog.Open(EMessageSeverity::Error);
+				CheckoutAndSaveWasSuccessful = false;
+			}
+			else
+			{
+				CheckedOutFiles.Add(Filename);
+			}
+		}
+
+		if( CheckoutAndSaveWasSuccessful )
+		{
+			//Print the JSON data out to the ref file.
+			FString OutputString;
+			TSharedRef< TJsonWriter<> > Writer = TJsonWriterFactory<>::Create( &OutputString );
+			FJsonSerializer::Serialize( Output, Writer );
+
+			if (!FFileHelper::SaveStringToFile(OutputString, *Filename, FFileHelper::EEncodingOptions::ForceUnicode))
+			{
+				// If we already checked out the file, but cannot write it, perhaps the user checked it in via perforce, so try to check it out again
+				if (bPreviouslyCheckedOut)
+				{
+					bPreviouslyCheckedOut = false;
+
+					if( !USourceControlHelpers::CheckOutOrAddFile(Filename) )
+					{
+						FFormatNamedArguments Arguments;
+						Arguments.Add( TEXT("Filename"), FText::FromString(Filename) );
+						// Use Source Control Message Log here because there might be other useful information in that log for the user.
+						FMessageLog SourceControlMessageLog("SourceControl");
+						SourceControlMessageLog.Error(FText::Format(LOCTEXT("CheckoutFailed", "Check out of file '{Filename}' failed."), Arguments));
+						SourceControlMessageLog.Notify(LOCTEXT("TranslationArchiveCheckoutFailed", "Failed to Check Out Translation Archive!"));
+						SourceControlMessageLog.Open(EMessageSeverity::Error);
+						CheckoutAndSaveWasSuccessful = false;
+
+						CheckedOutFiles.Remove(Filename);
+					}
+				}
+
+				FFormatNamedArguments Arguments;
+				Arguments.Add( TEXT("Filename"), FText::FromString(Filename) );
+				FMessageLog TranslationEditorMessageLog("TranslationEditor");
+				TranslationEditorMessageLog.Error(FText::Format(LOCTEXT("WriteFileFailed", "Failed to write localization entries to file '{Filename}'."), Arguments));
+				TranslationEditorMessageLog.Notify(LOCTEXT("FileWriteFailed", "Failed to Write Translations to File!"));
+				TranslationEditorMessageLog.Open(EMessageSeverity::Error);
+				CheckoutAndSaveWasSuccessful = false;
+			}
+		}
+	}
+	else
+	{
+		CheckoutAndSaveWasSuccessful = false;
+	}
+
+	// If this is the first time, let the user know the file was checked out
+	if (!bPreviouslyCheckedOut && CheckoutAndSaveWasSuccessful)
+	{
+		struct Local
+		{
+			/**
+			* Called by our notification's hyperlink to open the Source Control message log
+			*/
+			static void OpenSourceControlMessageLog(  )
+			{
+				FMessageLog("SourceControl").Open();
+			}
+		};
+
+		FFormatNamedArguments Arguments;
+		Arguments.Add( TEXT("Filename"), FText::FromString(Filename) );
+
+		// Make a note in the Source Control log, including a note to check in the file later via source control application
+		FMessageLog TranslationEditorMessageLog("SourceControl");
+		TranslationEditorMessageLog.Info(FText::Format(LOCTEXT("TranslationArchiveCheckedOut", "Successfully checked out and saved translation archive '{Filename}'. Please check-in this file later via your revision control application."), Arguments));
+
+		// Display notification that save was successful, along with a link to the Source Control log so the user can see the above message.
+		FNotificationInfo Info( LOCTEXT("ArchiveCheckedOut", "Translation Archive Successfully Checked Out and Saved.") );
+		Info.ExpireDuration = 5;
+		Info.Hyperlink = FSimpleDelegate::CreateStatic(&Local::OpenSourceControlMessageLog);
+		Info.HyperlinkText = LOCTEXT("ShowMessageLogHyperlink", "Show Message Log");
+		Info.bFireAndForget = true;
+		Info.bUseSuccessFailIcons = true;
+		Info.Image = FAppStyle::GetBrush(TEXT("NotificationList.SuccessImage"));
+		FSlateNotificationManager::Get().AddNotification(Info);
+	}
+
+	return CheckoutAndSaveWasSuccessful;
+}
+
+void FTranslationDataManager::GetHistoryForTranslationUnits()
+{
+	GWarn->BeginSlowTask(LOCTEXT("LoadingSourceControlHistory", "Loading Translation History from Revision Control..."), true);
+
+	TArray<UTranslationUnit*>& TranslationUnits = AllTranslations;
+	const FString& InManifestFilePath = OpenedManifestFilePath;
+
+	// Unload any previous history information, going to retrieve it all again.
+	UnloadHistoryInformation();
+
+	// Force history update
+	ISourceControlProvider& SourceControlProvider = ISourceControlModule::Get().GetProvider();
+	TSharedRef<FUpdateStatus, ESPMode::ThreadSafe> UpdateStatusOperation = ISourceControlOperation::Create<FUpdateStatus>();
+	UpdateStatusOperation->SetUpdateHistory( true );
+	ECommandResult::Type Result = SourceControlProvider.Execute(UpdateStatusOperation, InManifestFilePath);
+	bool bGetHistoryFromSourceControlSucceeded = Result == ECommandResult::Succeeded;
+
+	// Now we can get information about the file's history from the source control state, retrieve that
+	TArray<FString> Files;
+	TArray< TSharedRef<ISourceControlState, ESPMode::ThreadSafe> > States;
+	Files.Add(InManifestFilePath);
+	Result = SourceControlProvider.GetState( Files,  States, EStateCacheUsage::ForceUpdate );
+	bGetHistoryFromSourceControlSucceeded = bGetHistoryFromSourceControlSucceeded && (Result == ECommandResult::Succeeded);
+	FSourceControlStatePtr SourceControlState;
+	if (States.Num() == 1)
+	{
+		SourceControlState = States[0];
+	}
+
+	// If all the source control operations went ok, continue
+	if (bGetHistoryFromSourceControlSucceeded && SourceControlState.IsValid())
+	{
+		int32 HistorySize = SourceControlState->GetHistorySize();
+
+		for (int HistoryItemIndex = HistorySize-1; HistoryItemIndex >=0; --HistoryItemIndex)
+		{
+			GWarn->StatusUpdate(HistorySize - HistoryItemIndex, HistorySize, FText::Format(LOCTEXT("LoadingOldManifestRevisionNumber", "Loading Translation History from Manifest Revision {0} of {1} from Revision Control..."), FText::AsNumber(HistorySize - HistoryItemIndex), FText::AsNumber(HistorySize)));
+
+			TSharedPtr<ISourceControlRevision, ESPMode::ThreadSafe> Revision = SourceControlState->GetHistoryItem(HistoryItemIndex);
+			if(Revision.IsValid())
+			{
+				FString ManifestFullPath = FPaths::ConvertRelativePathToFull(InManifestFilePath);
+				FString EngineFullPath = FPaths::ConvertRelativePathToFull(FPaths::EngineContentDir());
+
+				bool IsEngineManifest = false;
+				if (ManifestFullPath.StartsWith(EngineFullPath))
+				{
+					IsEngineManifest = true;
+				}
+
+				FString ProjectName;
+				FString SavedDir;	// Store these cached translation history files in the saved directory
+				if (IsEngineManifest)
+				{
+					ProjectName = TEXT("Engine");
+					SavedDir = FPaths::EngineSavedDir();
+				}
+				else
+				{
+					ProjectName = FApp::GetProjectName();
+					SavedDir = FPaths::ProjectSavedDir();
+				}
+
+				FString TempFileName = SavedDir / TEXT("CachedTranslationHistory") / FString::Printf(TEXT("Manifest-%s-%s-Rev-%d"), *ProjectName, *FPaths::GetBaseFilename(InManifestFilePath), Revision->GetRevisionNumber());
+				
+				if (!FPaths::FileExists(TempFileName))	// Don't bother syncing again if we already have this manifest version cached locally
+				{
+					Revision->Get(TempFileName);
+				}
+
+				TSharedPtr< FInternationalizationManifest > OldManifestPtr = ReadManifest( TempFileName );
+				if (OldManifestPtr.IsValid())	// There may be corrupt manifests in the history, so ignore them.
+				{
+					TSharedRef< FInternationalizationManifest > OldManifest = OldManifestPtr.ToSharedRef();
+
+					for (UTranslationUnit* TranslationUnit : TranslationUnits)
+					{
+						if(TranslationUnit != nullptr && TranslationUnit->Contexts.Num() > 0)
+						{
+							for (FTranslationContextInfo& ContextInfo : TranslationUnit->Contexts)
+							{
+								FString PreviousSourceText;
+
+								// If we already have history, then compare against the newest history so far
+								if (ContextInfo.Changes.Num() > 0)
+								{
+									PreviousSourceText = ContextInfo.Changes[0].Source;
+								}
+
+								TSharedPtr< FManifestEntry > OldManifestEntryPtr = OldManifest->FindEntryByKey(TranslationUnit->Namespace, ContextInfo.Key);
+								if (!OldManifestEntryPtr.IsValid())
+								{
+									// If this version of the manifest didn't know anything about this string, move onto the next
+									continue;
+								}
+
+								// Always add first instance of this string, and then add any versions that changed since
+								if (ContextInfo.Changes.Num() == 0 || !OldManifestEntryPtr->Source.Text.Equals(PreviousSourceText))
+								{
+									TSharedPtr< FArchiveEntry > OldArchiveEntry = ArchivePtr->FindEntryByKey(OldManifestEntryPtr->Namespace, ContextInfo.Key, nullptr);
+									if (OldArchiveEntry.IsValid())
+									{
+										FTranslationChange Change;
+										Change.Source = OldManifestEntryPtr->Source.Text;
+										Change.Translation = OldArchiveEntry->Translation.Text;
+										Change.DateAndTime = Revision->GetDate();
+										Change.Version = FString::FromInt(Revision->GetRevisionNumber());
+										ContextInfo.Changes.Insert(Change, 0);
+									}
+								}
+							}
+						}
+					}
+				}
+				else // OldManifestPtr.IsValid() is false
+				{
+					FFormatNamedArguments Arguments;
+					Arguments.Add(TEXT("ManifestFilePath"), FText::FromString(InManifestFilePath));
+					Arguments.Add( TEXT("ManifestRevisionNumber"), FText::AsNumber(Revision->GetRevisionNumber()) );
+					FMessageLog TranslationEditorMessageLog("TranslationEditor");
+					TranslationEditorMessageLog.Warning(FText::Format(LOCTEXT("PreviousManifestCorrupt", "Previous revision {ManifestRevisionNumber} of {ManifestFilePath} failed to load correctly. Ignoring."), Arguments));
+				}
+			}
+		}
+	}
+	// If source control operations failed, display error message
+	else // (bGetHistoryFromSourceControlSucceeded && SourceControlState.IsValid()) is false
+	{
+		FFormatNamedArguments Arguments;
+		Arguments.Add(TEXT("ManifestFilePath"), FText::FromString(InManifestFilePath));
+		FMessageLog TranslationEditorMessageLog("SourceControl");
+		TranslationEditorMessageLog.Warning(FText::Format(LOCTEXT("SourceControlStateQueryFailed", "Failed to query revision control state of file {ManifestFilePath}."), Arguments));
+		TranslationEditorMessageLog.Notify(LOCTEXT("RetrieveTranslationHistoryFailed", "Unable to Retrieve Translation History from Revision Control!"));
+	}
+
+
+	// Go though all translation units
+	for (int32 CurrentTranslationUnitIndex = 0; CurrentTranslationUnitIndex < TranslationUnits.Num(); ++CurrentTranslationUnitIndex)
+	{
+		UTranslationUnit* TranslationUnit = TranslationUnits[CurrentTranslationUnitIndex];
+		if (TranslationUnit != nullptr)
+		{
+			if (TranslationUnit->Translation.IsEmpty())
+			{
+				bool bHasTranslationHistory = false;
+				int32 MostRecentNonNullTranslationIndex = -1;
+				int32 ContextForRecentTranslation = -1;
+
+				// Check all contexts for history
+				for (int32 ContextIndex = 0; ContextIndex < TranslationUnit->Contexts.Num(); ++ContextIndex)
+				{
+					for (int32 ChangeIndex = 0; ChangeIndex < TranslationUnit->Contexts[ContextIndex].Changes.Num(); ++ChangeIndex)
+					{
+						if (!(TranslationUnit->Contexts[ContextIndex].Changes[ChangeIndex].Translation.IsEmpty()))
+						{
+							bHasTranslationHistory = true;
+							MostRecentNonNullTranslationIndex = ChangeIndex;
+							ContextForRecentTranslation = ContextIndex;
+							break;
+						}
+					}
+
+					if (bHasTranslationHistory)
+					{
+						break;
+					}
+				}
+
+				// If we have history, but current translation is empty, this goes in the Needs Review tab
+				if (bHasTranslationHistory)
+				{
+					// Offer the most recent translation (for the first context in the list) as a suggestion or starting point (not saved unless user checks "Has Been Reviewed")
+					TranslationUnit->Translation = TranslationUnit->Contexts[ContextForRecentTranslation].Changes[MostRecentNonNullTranslationIndex].Translation;
+					TranslationUnit->HasBeenReviewed = false;
+
+					// Move from Untranslated to review
+					if (Untranslated.Contains(TranslationUnit))
+					{
+						Untranslated.Remove(TranslationUnit);
+					}
+					if (!Review.Contains(TranslationUnit))
+					{
+						Review.Add(TranslationUnit);
+					}
+				}
+			}
+		}
+	}
+
+
+	GWarn->EndSlowTask();
+}
+
+void FTranslationDataManager::HandlePropertyChanged(FName PropertyName)
+{
+	// When a property changes, write the data so we don't lose changes if user forgets to save or editor crashes
+	WriteTranslationData();
+}
+
+void FTranslationDataManager::PreviewAllTranslationsInEditor(ULocalizationTarget* LocalizationTarget)
+{
+	FString ManifestFullPath = FPaths::ConvertRelativePathToFull(OpenedManifestFilePath);
+	FString EngineFullPath = FPaths::ConvertRelativePathToFull(FPaths::EngineContentDir());
+
+	bool IsEngineManifest = false;
+	if (ManifestFullPath.StartsWith(EngineFullPath))
+	{
+		IsEngineManifest = true;
+	}
+
+	if (LocalizationTarget != nullptr)
+	{
+		const FString ConfigFilePath = LocalizationConfigurationScript::GetRegenerateResourcesConfigPath(LocalizationTarget);
+		LocalizationConfigurationScript::GenerateRegenerateResourcesConfigFile(LocalizationTarget).Write(ConfigFilePath);
+		FTextLocalizationResourceGenerator::GenerateLocResAndUpdateLiveEntriesFromConfig(ConfigFilePath, EGenerateLocResFlags::None);
+	}
+	else
+	{
+		FText ErrorNotify = LOCTEXT("PreviewAllTranslationsInEditorFail", "Failed to preview translations in Editor!");
+		FMessageLog TranslationEditorMessageLog("TranslationEditor");
+		TranslationEditorMessageLog.Error(ErrorNotify);
+		TranslationEditorMessageLog.Notify(ErrorNotify);
+	}
+}
+
+void FTranslationDataManager::PopulateSearchResultsUsingFilter(const FString& SearchFilter)
+{
+	SearchResults.Empty();
+
+	for (UTranslationUnit* TranslationUnit : AllTranslations)
+	{
+		if (TranslationUnit != nullptr)
+		{
+			bool bAdded = false;
+			if (TranslationUnit->Source.Contains(SearchFilter) ||
+				TranslationUnit->Translation.Contains(SearchFilter) || 
+				TranslationUnit->Namespace.Contains(SearchFilter))
+			{
+				SearchResults.Add(TranslationUnit);
+				bAdded = true;
+			}
+			
+			for (FTranslationContextInfo CurrentContext : TranslationUnit->Contexts)
+			{
+				if (!bAdded && 
+					(CurrentContext.Context.Contains(SearchFilter) ||
+					CurrentContext.Key.Contains(SearchFilter)))
+				{
+					SearchResults.Add(TranslationUnit);
+					break;
+				}
+			}
+		}
+	}
+}
+
+void FTranslationDataManager::LoadFromArchive(TArray<UTranslationUnit*>& InTranslationUnits, bool bTrackChanges /*= false*/, bool bReloadFromFile /*=false*/)
+{
+	GWarn->BeginSlowTask(LOCTEXT("LoadingArchiveEntries", "Loading Entries from Translation Archive..."), true);
+
+	if (bReloadFromFile)
+	{
+		ArchivePtr = ReadArchive(OpenedArchiveFilePath);
+		NativeArchivePtr = NativeArchiveFilePath != OpenedArchiveFilePath ? ReadArchive(NativeArchiveFilePath) : ArchivePtr;
+	}
+
+	if (ArchivePtr.IsValid())
+	{
+		const TSharedRef< FInternationalizationArchive > Archive = ArchivePtr.ToSharedRef();
+
+		// Make a local copy of this array before we empty the arrays below (we might have been passed AllTranslations array)
+		TArray<UTranslationUnit*> TranslationUnits;
+		TranslationUnits.Append(InTranslationUnits);
+
+		AllTranslations.Empty();
+		Untranslated.Empty();
+		Review.Empty();
+		Complete.Empty();
+		ChangedOnImport.Empty();
+
+		for (int32 CurrentTranslationUnitIndex = 0; CurrentTranslationUnitIndex < TranslationUnits.Num(); ++CurrentTranslationUnitIndex)
+		{
+			UTranslationUnit* TranslationUnit = TranslationUnits[CurrentTranslationUnitIndex];
+			if (TranslationUnit != nullptr)
+			{
+				if (!TranslationUnit->IsRooted())
+				{
+					TranslationUnit->AddToRoot(); // Disable garbage collection for UTranslationUnit objects
+				}
+				AllTranslations.Add(TranslationUnit);
+
+				GWarn->StatusUpdate(CurrentTranslationUnitIndex, TranslationUnits.Num(), FText::Format(LOCTEXT("LoadingCurrentArchiveEntries", "Loading Entry {0} of {1} from Translation Archive..."), FText::AsNumber(CurrentTranslationUnitIndex), FText::AsNumber(TranslationUnits.Num())));
+
+				TSharedPtr<FArchiveEntry> ArchiveEntry = Archive->FindEntryByKey(TranslationUnit->Namespace, TranslationUnit->Key, TranslationUnit->KeyMetaDataObject);
+				if (ArchiveEntry.IsValid())
+				{
+					const FString PreviousTranslation = TranslationUnit->Translation;
+					TranslationUnit->Translation.Reset();
+
+					if (ArchiveEntry->Translation.Text.IsEmpty())
+					{
+						// No current translation - try and find an historic one
+						bool bHasTranslationHistory = false;
+						int32 MostRecentNonNullTranslationIndex = -1;
+						int32 ContextForRecentTranslation = -1;
+
+						for (int32 ContextIndex = 0; ContextIndex < TranslationUnit->Contexts.Num(); ++ContextIndex)
+						{
+							for (int32 ChangeIndex = 0; ChangeIndex < TranslationUnit->Contexts[ContextIndex].Changes.Num(); ++ChangeIndex)
+							{
+								if (!(TranslationUnit->Contexts[ContextIndex].Changes[ChangeIndex].Translation.IsEmpty()))
+								{
+									bHasTranslationHistory = true;
+									MostRecentNonNullTranslationIndex = ChangeIndex;
+									ContextForRecentTranslation = ContextIndex;
+									break;
+								}
+							}
+
+							if (bHasTranslationHistory)
+							{
+								break;
+							}
+						}
+
+						// If we have history, but current translation is empty, this goes in the Needs Review tab
+						if (bHasTranslationHistory)
+						{
+							// Offer the most recent translation (for the first context in the list) as a suggestion or starting point (not saved unless user checks "Has Been Reviewed")
+							TranslationUnit->Translation = TranslationUnit->Contexts[ContextForRecentTranslation].Changes[MostRecentNonNullTranslationIndex].Translation;
+							Review.Add(TranslationUnit);
+						}
+						else
+						{
+							Untranslated.Add(TranslationUnit);
+						}
+					}
+					else if (!ArchiveEntry->Source.Text.Equals(TranslationUnit->Source, ESearchCase::CaseSensitive))
+					{
+						// If we have a stale translation, this goes in the Needs Review tab
+						TranslationUnit->Translation = ArchiveEntry->Translation.Text;
+						Review.Add(TranslationUnit);
+					}
+					else
+					{
+						TranslationUnit->Translation = ArchiveEntry->Translation.Text;
+						TranslationUnit->HasBeenReviewed = true;
+						Complete.Add(TranslationUnit);
+					}
+
+					// Add to changed array if we're tracking changes (i.e. when we import from .po files)
+					if (bTrackChanges)
+					{
+						if (PreviousTranslation != TranslationUnit->Translation)
+						{
+							FString PreviousTranslationTrimmed = PreviousTranslation;
+							PreviousTranslationTrimmed.TrimStartAndEndInline();
+							FString CurrentTranslationTrimmed = TranslationUnit->Translation;
+							CurrentTranslationTrimmed.TrimStartAndEndInline();
+							// Ignore changes to only whitespace at beginning and/or end of string on import
+							if (PreviousTranslationTrimmed.Equals(CurrentTranslationTrimmed))
+							{
+								TranslationUnit->Translation = PreviousTranslation;
+							}
+							else
+							{
+								ChangedOnImport.Add(TranslationUnit);
+								TranslationUnit->TranslationBeforeImport = PreviousTranslation;
+							}
+						}
+					}
+				}
+			}
+		}
+
+	}
+	else  // ArchivePtr.IsValid() is false
+	{
+		FFormatNamedArguments Arguments;
+		Arguments.Add(TEXT("ArchiveFilePath"), FText::FromString(OpenedArchiveFilePath));
+		FMessageLog TranslationEditorMessageLog("TranslationEditor");
+		TranslationEditorMessageLog.Error(FText::Format(LOCTEXT("FailedToLoadCurrentArchive", "Failed to load most current translation archive ({ArchiveFilePath}), unable to load translations."), Arguments));
+		TranslationEditorMessageLog.Notify(LOCTEXT("TranslationLoadError", "Error Loading Translations!"));
+		TranslationEditorMessageLog.Open(EMessageSeverity::Error);
+	}
+	
+	GWarn->EndSlowTask();
+}
+
+void FTranslationDataManager::RemoveTranslationUnitArrayfromRoot(TArray<UTranslationUnit*>& TranslationUnits)
+{
+	for (UTranslationUnit* TranslationUnit : TranslationUnits)
+	{
+		TranslationUnit->RemoveFromRoot();
+	}
+}
+
+void FTranslationDataManager::UnloadHistoryInformation()
+{
+	TArray<UTranslationUnit*>& TranslationUnits = AllTranslations;
+
+	for (int32 CurrentTranslationUnitIndex = 0; CurrentTranslationUnitIndex < TranslationUnits.Num(); ++CurrentTranslationUnitIndex)
+	{
+		UTranslationUnit* TranslationUnit = TranslationUnits[CurrentTranslationUnitIndex];
+		if (TranslationUnit != nullptr)
+		{
+			// If HasBeenReviewed is false, this is a suggestion translation from a previous translation for the same Namespace/Key pair
+			if (!TranslationUnit->HasBeenReviewed)
+			{
+				if (!Untranslated.Contains(TranslationUnit))
+				{
+					Untranslated.Add(TranslationUnit);
+				}
+				if (Review.Contains(TranslationUnit))
+				{
+					Review.Remove(TranslationUnit);
+				}
+
+				// Erase previously suggested translation from history (it has not been reviewed)
+				TranslationUnit->Translation.Empty();
+
+				// Remove all history entries
+				for (FTranslationContextInfo Context : TranslationUnit->Contexts)
+				{
+					Context.Changes.Empty();
+				}
+			}
+		}
+	}
+}
+
+bool FTranslationDataManager::SaveSelectedTranslations(TArray<UTranslationUnit*> TranslationUnitsToSave, bool bSaveChangesToTranslationService)
+{
+	bool bSucceeded = true;
+	
+	TMap<FString, TSharedPtr<TArray<UTranslationUnit*>>> TextsToSavePerProject;
+
+	// Regroup the translations to save by project
+	for (UTranslationUnit* TextToSave : TranslationUnitsToSave)
+	{
+		FString LocresFilePath = TextToSave->LocresPath;
+		if (!LocresFilePath.IsEmpty())
+		{
+			if (!TextsToSavePerProject.Contains(LocresFilePath))
+			{
+				TextsToSavePerProject.Add(LocresFilePath, MakeShareable(new TArray<UTranslationUnit*>()));
+			}
+
+			TSharedPtr<TArray<UTranslationUnit*>> ProjectArray = TextsToSavePerProject.FindRef(LocresFilePath);
+			ProjectArray->Add(TextToSave);
+		}
+	}
+
+	for (auto TextIt = TextsToSavePerProject.CreateIterator(); TextIt; ++TextIt)
+	{
+		auto Item = *TextIt;
+		FString CurrentLocResPath = Item.Key;
+		FString ManifestAndArchiveName = FPaths::GetBaseFilename(CurrentLocResPath);
+
+		FString ArchiveFilePath = FPaths::GetPath(CurrentLocResPath);
+		FString CultureName = FPaths::GetBaseFilename(ArchiveFilePath);
+		FString ManifestPath = FPaths::GetPath(ArchiveFilePath);
+		FString ArchiveFullPath = ArchiveFilePath / ManifestAndArchiveName + ".archive";
+		FString ManifestFullPath = ManifestPath / ManifestAndArchiveName + ".manifest";
+		FString EngineFullPath = FPaths::ConvertRelativePathToFull(FPaths::EngineContentDir());
+		bool IsEngineManifest = false;
+		if (ManifestFullPath.StartsWith(EngineFullPath))
+		{
+			IsEngineManifest = true;
+		}
+
+		ULocalizationTarget* LocalizationTarget = ILocalizationModule::Get().GetLocalizationTargetByName(ManifestAndArchiveName, IsEngineManifest);
+
+		if (LocalizationTarget && FPaths::FileExists(ManifestFullPath) && FPaths::FileExists(ArchiveFullPath))
+		{
+			FString NativeCultureName;
+			if (LocalizationTarget->Settings.SupportedCulturesStatistics.IsValidIndex(LocalizationTarget->Settings.NativeCultureIndex))
+			{
+				NativeCultureName = LocalizationTarget->Settings.SupportedCulturesStatistics[LocalizationTarget->Settings.NativeCultureIndex].CultureName;
+			}
+			FString NativeArchiveFullPath = ManifestPath / NativeCultureName / ManifestAndArchiveName + ".archive";
+
+			TSharedRef<FTranslationDataManager> DataManager = MakeShareable(new FTranslationDataManager(ManifestFullPath, NativeArchiveFullPath, ArchiveFullPath));
+
+			if (DataManager->GetLoadedSuccessfully())
+			{
+				FPortableObjectFormatDOM PortableObjectDom;
+				PortableObjectDom.SetProjectName(ManifestAndArchiveName);
+				PortableObjectDom.SetLanguage(CultureName);
+				PortableObjectDom.CreateNewHeader(GetDefault<UGeneralProjectSettings>()->CopyrightNotice);
+				PortableObjectPipeline::UpdatePOFileHeaderForSettings(PortableObjectDom, LocalizationTarget->Settings.ExportSettings.CollapseMode, LocalizationTarget->Settings.ExportSettings.POFormat);
+
+				TArray<UTranslationUnit*>& TranslationsArray = DataManager->GetAllTranslationsArray();
+				TSharedPtr<TArray<UTranslationUnit*>> EditedItems = Item.Value;
+
+				// For each edited item belonging to this manifest/archive pair
+				for (auto EditedItemIt = EditedItems->CreateIterator(); EditedItemIt; ++EditedItemIt)
+				{
+					UTranslationUnit* EditedItem = *EditedItemIt;
+
+					// Search all translations for the one that matches this FText
+					for (UTranslationUnit* Translation : TranslationsArray)
+					{
+						bool bFoundMatchingTranslation = false;
+
+						// If namespace matches...
+						if (Translation->Namespace.Equals(EditedItem->Namespace))
+						{
+							// And key matches...
+							for (const FTranslationContextInfo& ContextInfo : Translation->Contexts)
+							{
+								if (ContextInfo.Key.Equals(EditedItem->Key))
+								{
+									bFoundMatchingTranslation = true;
+
+									// Update the translation in TranslationDataManager, and finish searching these translations
+									Translation->Translation = EditedItem->Translation;
+
+									// Add the PO entry
+									{
+										TSharedRef<FPortableObjectEntry> PoEntry = MakeShareable(new FPortableObjectEntry());
+										PortableObjectPipeline::PopulateBasicPOFileEntry(*PoEntry, Translation->Namespace, ContextInfo.Key, nullptr, Translation->Source, Translation->Translation, LocalizationTarget->Settings.ExportSettings.CollapseMode, LocalizationTarget->Settings.ExportSettings.POFormat);
+
+										//@TODO: We support additional metadata entries that can be translated.  How do those fit in the PO file format?  Ex: isMature
+										PoEntry->AddReference(ContextInfo.Context); // Source location.
+
+										PoEntry->AddExtractedComment(PortableObjectPipeline::GetConditionedKeyForExtractedComment(ContextInfo.Key)); // "Notes from Programmer" in the form of the Key.
+										PoEntry->AddExtractedComment(PortableObjectPipeline::GetConditionedReferenceForExtractedComment(ContextInfo.Context)); // "Notes from Programmer" in the form of the Source Location, since this comes in handy too and OneSky doesn't properly show references, only comments.
+
+										PortableObjectDom.AddEntry(PoEntry);
+									}
+									break;
+								}
+							}
+
+							if (bFoundMatchingTranslation)
+							{
+								break;
+							}
+						}
+					}
+				}
+
+				if (bSaveChangesToTranslationService)
+				{
+					FString UploadFilePath = FPaths::ProjectSavedDir() / "Temp" / CultureName / ManifestAndArchiveName + ".po";
+					FFileHelper::SaveStringToFile(PortableObjectDom.ToString(), *UploadFilePath);
+
+					FGuid LocalizationTargetGuid = LocalizationTarget->Settings.Guid;
+
+					ILocalizationServiceProvider& Provider = ILocalizationServiceModule::Get().GetProvider();
+					TSharedRef<FUploadLocalizationTargetFile, ESPMode::ThreadSafe> UploadTargetFileOp = ILocalizationServiceOperation::Create<FUploadLocalizationTargetFile>();
+					UploadTargetFileOp->SetInTargetGuid(LocalizationTargetGuid);
+					UploadTargetFileOp->SetInLocale(CultureName);
+					FPaths::MakePathRelativeTo(UploadFilePath, *FPaths::ProjectDir());
+					UploadTargetFileOp->SetInRelativeInputFilePathAndName(UploadFilePath);
+					UploadTargetFileOp->SetPreserveAllText(true);
+
+					Provider.Execute(UploadTargetFileOp, TArray<FLocalizationServiceTranslationIdentifier>(), ELocalizationServiceOperationConcurrency::Asynchronous, FLocalizationServiceOperationComplete::CreateStatic(&FTranslationDataManager::SaveSelectedTranslationsToTranslationServiceCallback));
+				}
+
+			}
+			else
+			{
+				bSucceeded = false;
+			}
+
+			// Save the data to file, and preview in editor
+			bSucceeded = bSucceeded && DataManager->WriteTranslationData();
+			DataManager->PreviewAllTranslationsInEditor(LocalizationTarget);
+		}
+		else
+		{
+			bSucceeded = false;
+		}
+	}
+
+	return bSucceeded;
+}
+
+void FTranslationDataManager::SaveSelectedTranslationsToTranslationServiceCallback(const FLocalizationServiceOperationRef& Operation, ELocalizationServiceOperationCommandResult::Type Result)
+{
+	TSharedPtr<FUploadLocalizationTargetFile, ESPMode::ThreadSafe> UploadLocalizationTargetOp = StaticCastSharedRef<FUploadLocalizationTargetFile>(Operation);
+	bool bError = !(Result == ELocalizationServiceOperationCommandResult::Succeeded);
+	FText ErrorText = FText::GetEmpty();
+	FGuid InTargetGuid;
+	FString InLocale;
+	FString InRelativeOutputFilePathAndName;
+	FString TargetName = "";
+	FString TargetPath = "";
+	FString CultureName = "";
+	if (UploadLocalizationTargetOp.IsValid())
+	{
+		ErrorText = UploadLocalizationTargetOp->GetOutErrorText();
+		InTargetGuid = UploadLocalizationTargetOp->GetInTargetGuid();
+		InLocale = UploadLocalizationTargetOp->GetInLocale();
+		InRelativeOutputFilePathAndName = UploadLocalizationTargetOp->GetInRelativeInputFilePathAndName();
+		TargetName = FPaths::GetBaseFilename(InRelativeOutputFilePathAndName);
+		TargetPath = FPaths::GetPath(InRelativeOutputFilePathAndName);
+		CultureName = FPaths::GetBaseFilename(TargetPath);
+	}
+
+	// Try to get display name
+	FInternationalization& I18N = FInternationalization::Get();
+	FCulturePtr CulturePtr = I18N.GetCulture(CultureName);
+	FString CultureDisplayName = CultureName;
+	if (CulturePtr.IsValid())
+	{
+		CultureName = CulturePtr->GetDisplayName();
+	}
+
+	if (!bError && ErrorText.IsEmpty())
+	{
+		FText SuccessText = FText::Format(LOCTEXT("SaveSelectedTranslationsToTranslationServiceSuccess", "{0} translations for {1} target uploaded for processing to Translation Service."), FText::FromString(CultureDisplayName), FText::FromString(TargetName));
+		FMessageLog TranslationEditorMessageLog("TranslationEditor");
+		TranslationEditorMessageLog.Info(SuccessText);
+		TranslationEditorMessageLog.Notify(SuccessText, EMessageSeverity::Info, true);
+	}
+	else
+	{
+		if (ErrorText.IsEmpty())
+		{
+			ErrorText = LOCTEXT("SaveToLocalizationServiceUnspecifiedError", "An unspecified error occured when trying to save to the Localization Service.");
+		}
+
+		FText ErrorNotify = FText::Format(LOCTEXT("SaveSelectedTranslationsToTranslationServiceFail", "{0} translations for {1} target failed to save to Translation Service!"), FText::FromString(CultureDisplayName), FText::FromString(TargetName));
+		FMessageLog TranslationEditorMessageLog("TranslationEditor");
+		TranslationEditorMessageLog.Error(ErrorNotify);
+		TranslationEditorMessageLog.Error(ErrorText);
+		TranslationEditorMessageLog.Notify(ErrorNotify);
+	}
+}
+
+#undef LOCTEXT_NAMESPACE
diff --git a/Engine/Source/Developer/TranslationEditor/Private/TranslationDataManager.h b/Engine/Source/Developer/TranslationEditor/Private/TranslationDataManager.h
new file mode 100644
--- /dev/null
+++ b/Engine/Source/Developer/TranslationEditor/Private/TranslationDataManager.h
@@ -0,0 +1,143 @@
+// Copyright Epic Games, Inc. All Rights Reserved.
+
+#pragma once
+
+#include "CoreMinimal.h"
+#include "UObject/WeakObjectPtr.h"
+#include "ILocalizationServiceProvider.h"
+
+class FInternationalizationArchive;
+class FInternationalizationManifest;
+class FJsonObject;
+class ULocalizationTarget;
+class UTranslationUnit;
+
+class FTranslationDataManager : public TSharedFromThis<FTranslationDataManager>
+{
+
+public:
+
+	FTranslationDataManager(const FString& InManifestFilePath, const FString& InNativeArchiveFilePath, const FString& InArchiveFilePath);
+	FTranslationDataManager(ULocalizationTarget* const LocalizationTarget, const FString& CulturetoEdit);
+
+	virtual ~FTranslationDataManager();
+
+	TArray<UTranslationUnit*>& GetAllTranslationsArray()
+	{
+		return AllTranslations;
+	}
+
+	TArray<UTranslationUnit*>& GetUntranslatedArray()
+	{
+		return Untranslated;
+	}
+
+	TArray<UTranslationUnit*>& GetReviewArray()
+	{
+		return Review;
+	}
+
+	TArray<UTranslationUnit*>& GetCompleteArray()
+	{
+		return Complete;
+	}
+
+	TArray<UTranslationUnit*>& GetSearchResultsArray()
+	{
+		return SearchResults;
+	}
+
+	TArray<UTranslationUnit*>& GetChangedOnImportArray()
+	{
+		return ChangedOnImport;
+	}
+	
+	/** Write the translation data in memory out to .archive file (check out the .archive file first if necessary) 
+	* @param bForceWrite Whether or not to force a write even if the data hasn't changed
+	* @param return Whether or not the write succeeded (true is succeeded)
+	*/
+	bool WriteTranslationData(bool bForceWrite = false);
+
+	/** Delegate called when a TranslationDataObject property is changed */
+	void HandlePropertyChanged(FName PropertyName);
+
+	/** Regenerate and reload archives to reflect modifications in the UI */
+	void PreviewAllTranslationsInEditor(ULocalizationTarget* LocalizationTarget);
+
+	/** Put items in the Search Array if they match this filter */
+	void PopulateSearchResultsUsingFilter(const FString& SearchFilter);
+
+	/** Load (or reload) Translations from Archive file */ 
+	void LoadFromArchive(TArray<UTranslationUnit*>& InTranslationUnits, bool bTrackChanges = false, bool bReloadFromFile = false);
+
+	/** Get the history data for all translation units */
+	void GetHistoryForTranslationUnits();
+
+	/** Unload History information (in the case we reload it) */
+	void UnloadHistoryInformation();
+
+	/** Save the specified translations */
+	static bool SaveSelectedTranslations(TArray<UTranslationUnit*> TranslationUnitsToSave, bool bSaveChangesToTranslationService = false);
+
+	/** Save the specified translations */
+	static void SaveSelectedTranslationsToTranslationServiceCallback(const FLocalizationServiceOperationRef& Operation, ELocalizationServiceOperationCommandResult::Type Result);
+
+	/** Whether or not we successfully loaded the .manifest and .archive */
+	bool GetLoadedSuccessfully()
+	{
+		return bLoadedSuccessfully;
+	}
+
+private:
+	void Initialize();
+
+	/** Take a path and a manifest name and return a manifest data structure */
+	TSharedPtr< FInternationalizationManifest > ReadManifest ( const FString& ManifestFilePath );
+	/** Retrieve an archive data structure from ArchiveFilePath */
+	TSharedPtr< FInternationalizationArchive > ReadArchive(const FString& ArchiveFilePath);
+
+	/** Write JSON file to text file */
+	bool WriteJSONToTextFile( TSharedRef<FJsonObject>& Output, const FString& Filename );
+
+	/** Removes each UTranslationUnit in the passed array from the root set, allowing it to be garbage collected */
+	void RemoveTranslationUnitArrayfromRoot( TArray<UTranslationUnit*>& TranslationUnits );
+
+	// Arrays containing the translation data
+	TArray<UTranslationUnit*> AllTranslations;
+	TArray<UTranslationUnit*> Untranslated;
+	TArray<UTranslationUnit*> Review;
+	TArray<UTranslationUnit*> Complete;
+	TArray<UTranslationUnit*> SearchResults;
+	TArray<UTranslationUnit*> ChangedOnImport;
+
+	/** Archive for the current project and native language */
+	TSharedPtr< FInternationalizationArchive > NativeArchivePtr;
+	/** Archive for the current project and translation language */
+	TSharedPtr< FInternationalizationArchive > ArchivePtr;
+	/** Manifest for the current project */
+	TSharedPtr< FInternationalizationManifest > ManifestAtHeadRevisionPtr;
+
+	/** Name of the manifest file */
+	FString ManifestName;
+	/** Path to the project */
+	FString ProjectPath;
+	/** Name of the archive file */
+	FString ArchiveName;
+	/** Path to the culture (language, sort of) we are targeting */
+	FString CulturePath;
+	/** Path to the manifest file **/
+	FString OpenedManifestFilePath;
+	/** Path to the native culture's archive file **/
+	FString NativeArchiveFilePath;
+	/** Path to the archive file **/
+	FString OpenedArchiveFilePath;
+
+	/** Files that are already checked out from Perforce **/
+	TArray<FString> CheckedOutFiles;
+
+	/** The localization target associated with the files being used/edited, if any. */
+	TWeakObjectPtr<ULocalizationTarget> AssociatedLocalizationTarget;
+
+	/** Whether or not we successfully loaded the .manifest and .archive */
+	bool bLoadedSuccessfully;
+};
diff --git a/Engine/Source/Developer/TranslationEditor/Private/TranslationEditor.cpp b/Engine/Source/Developer/TranslationEditor/Private/TranslationEditor.cpp
new file mode 100644
--- /dev/null
+++ b/Engine/Source/Developer/TranslationEditor/Private/TranslationEditor.cpp
@@ -0,0 +1,1484 @@
+// Copyright Epic Games, Inc. All Rights Reserved.
+#if WITH_EDITOR
+
+#include "TranslationEditor.h"
+#include "Misc/FeedbackContext.h"
+#include "Modules/ModuleManager.h"
+#include "Framework/Application/SlateApplication.h"
+#include "Widgets/Input/SButton.h"
+#include "TranslationEditorModule.h"
+#include "TranslationEditorMenu.h"
+#include "PropertyEditorModule.h"
+#include "PropertyPath.h"
+
+#include "Logging/MessageLog.h"
+
+#include "IPropertyTableRow.h"
+#include "DesktopPlatformModule.h"
+#include "IPropertyTableWidgetHandle.h"
+
+#include "Widgets/Input/SSearchBox.h"
+#include "Widgets/Docking/SDockTab.h"
+#include "ILocalizationServiceModule.h"
+#include "LocalizationCommandletTasks.h"
+#include "Framework/Notifications/NotificationManager.h"
+#include "Widgets/Notifications/SNotificationList.h"
+
+#define LOCTEXT_NAMESPACE "TranslationEditor"
+
+namespace TranslationEditorUtils
+{
+
+/** Get the filename used by the given font info */
+FString GetFontFilename(const FSlateFontInfo& InFontInfo)
+{
+	const FCompositeFont* const ResolvedCompositeFont = InFontInfo.GetCompositeFont();
+	return (ResolvedCompositeFont && ResolvedCompositeFont->DefaultTypeface.Fonts.Num() > 0)
+		? ResolvedCompositeFont->DefaultTypeface.Fonts[0].Font.GetFontFilename()
+		: FString();
+}
+
+} // namespace TranslationEditorUtils
+
+const FName FTranslationEditor::UntranslatedTabId( TEXT( "TranslationEditor_Untranslated" ) );
+const FName FTranslationEditor::ReviewTabId( TEXT( "TranslationEditor_Review" ) );
+const FName FTranslationEditor::CompletedTabId( TEXT( "TranslationEditor_Completed" ) );
+const FName FTranslationEditor::PreviewTabId( TEXT( "TranslationEditor_Preview" ) );
+const FName FTranslationEditor::ContextTabId( TEXT( "TranslationEditor_Context" ) );
+const FName FTranslationEditor::HistoryTabId( TEXT( "TranslationEditor_History" ) );
+const FName FTranslationEditor::SearchTabId( TEXT( "TranslationEditor_Search" ) );
+const FName FTranslationEditor::ChangedOnImportTabId( TEXT( "TranslationEditor_ChangedOnImport" ) );
+
+void FTranslationEditor::Initialize()
+{
+	// Set up delegate functions for the buttons/spinboxes in the custom font columns' headers
+	SourceColumn->SetOnChangeFontButtonClicked(FOnClicked::CreateSP(this, &FTranslationEditor::ChangeSourceFont_FReply));
+	SourceColumn->SetOnFontSizeValueCommitted(FOnInt32ValueCommitted::CreateSP(this, &FTranslationEditor::OnSourceFontSizeCommitt));
+	TranslationColumn->SetOnChangeFontButtonClicked(FOnClicked::CreateSP(this, &FTranslationEditor::ChangeTranslationTargetFont_FReply));
+	TranslationColumn->SetOnFontSizeValueCommitted(FOnInt32ValueCommitted::CreateSP(this, &FTranslationEditor::OnTranslationTargetFontSizeCommitt));
+}
+
+void FTranslationEditor::RegisterTabSpawners(const TSharedRef<class FTabManager>& InTabManager)
+{
+	WorkspaceMenuCategory = InTabManager->AddLocalWorkspaceMenuCategory(LOCTEXT("WorkspaceMenu_TranslationEditor", "Translation Editor"));
+	auto WorkspaceMenuCategoryRef = WorkspaceMenuCategory.ToSharedRef();
+
+	FAssetEditorToolkit::RegisterTabSpawners(InTabManager);
+
+	InTabManager->RegisterTabSpawner( UntranslatedTabId, FOnSpawnTab::CreateSP(this, &FTranslationEditor::SpawnTab_Untranslated) )
+		.SetDisplayName( LOCTEXT("UntranslatedTab", "Untranslated") )
+		.SetGroup( WorkspaceMenuCategoryRef );
+
+	InTabManager->RegisterTabSpawner( ReviewTabId, FOnSpawnTab::CreateSP(this, &FTranslationEditor::SpawnTab_Review) )
+		.SetDisplayName( LOCTEXT("ReviewTab", "Needs Review") )
+		.SetGroup( WorkspaceMenuCategoryRef );
+
+	InTabManager->RegisterTabSpawner( CompletedTabId, FOnSpawnTab::CreateSP(this, &FTranslationEditor::SpawnTab_Completed) )
+		.SetDisplayName( LOCTEXT("CompletedTab", "Completed") )
+		.SetGroup( WorkspaceMenuCategoryRef );
+
+	InTabManager->RegisterTabSpawner( PreviewTabId, FOnSpawnTab::CreateSP(this, &FTranslationEditor::SpawnTab_Preview) )
+		.SetDisplayName( LOCTEXT("PreviewTab", "Preview") )
+		.SetGroup( WorkspaceMenuCategoryRef );
+
+	InTabManager->RegisterTabSpawner( ContextTabId, FOnSpawnTab::CreateSP(this, &FTranslationEditor::SpawnTab_Context) )
+		.SetDisplayName( LOCTEXT("ContextTab", "Context") )
+		.SetGroup( WorkspaceMenuCategoryRef );
+
+	InTabManager->RegisterTabSpawner( HistoryTabId, FOnSpawnTab::CreateSP(this, &FTranslationEditor::SpawnTab_History) )
+		.SetDisplayName( LOCTEXT("HistoryTab", "History") )
+		.SetGroup( WorkspaceMenuCategoryRef );
+
+	InTabManager->RegisterTabSpawner( SearchTabId, FOnSpawnTab::CreateSP(this, &FTranslationEditor::SpawnTab_Search) )
+		.SetDisplayName(LOCTEXT("SearchTab", "Search"))
+		.SetGroup( WorkspaceMenuCategoryRef );
+
+	InTabManager->RegisterTabSpawner( ChangedOnImportTabId, FOnSpawnTab::CreateSP(this, &FTranslationEditor::SpawnTab_ChangedOnImport) )
+		.SetDisplayName(LOCTEXT("ChangedOnImportTab", "Changed On Import"))
+		.SetGroup( WorkspaceMenuCategoryRef );
+}
+
+void FTranslationEditor::UnregisterTabSpawners(const TSharedRef<class FTabManager>& InTabManager)
+{
+	InTabManager->UnregisterTabSpawner( UntranslatedTabId );
+	InTabManager->UnregisterTabSpawner( ReviewTabId );
+	InTabManager->UnregisterTabSpawner( CompletedTabId );
+	InTabManager->UnregisterTabSpawner( PreviewTabId );
+	InTabManager->UnregisterTabSpawner( ContextTabId );
+	InTabManager->UnregisterTabSpawner(HistoryTabId);
+	InTabManager->UnregisterTabSpawner(SearchTabId);
+	InTabManager->UnregisterTabSpawner(ChangedOnImportTabId);
+}
+
+void FTranslationEditor::InitTranslationEditor( const EToolkitMode::Type Mode, const TSharedPtr< class IToolkitHost >& InitToolkitHost )
+{	
+	TSharedRef<FTabManager::FLayout> StandaloneDefaultLayout = FTabManager::NewLayout( "Standalone_TranslationEditor_Layout_v2" )
+	->AddArea
+	(
+		FTabManager::NewPrimaryArea()
+		->SetOrientation(Orient_Vertical)
+		->Split
+		(
+			FTabManager::NewStack()
+			->SetSizeCoefficient(0.5)
+			->SetHideTabWell( false )
+			->AddTab( UntranslatedTabId, ETabState::OpenedTab )
+			->AddTab( ReviewTabId,  ETabState::OpenedTab )
+			->AddTab( CompletedTabId,  ETabState::OpenedTab )
+			->AddTab( SearchTabId, ETabState::ClosedTab )
+			->AddTab( ChangedOnImportTabId, ETabState::ClosedTab )
+		)
+		->Split
+		(
+			FTabManager::NewStack()
+			->SetSizeCoefficient(0.5)
+			->SetHideTabWell(false)
+			->AddTab(PreviewTabId, ETabState::OpenedTab)
+		)
+		->Split
+		(
+			FTabManager::NewSplitter()
+			->Split
+			(
+				FTabManager::NewStack()
+				->SetHideTabWell(false)
+				->AddTab(ContextTabId, ETabState::OpenedTab)
+			)
+			->Split
+			(
+				FTabManager::NewStack()
+				->SetHideTabWell(false)
+				->AddTab(HistoryTabId, ETabState::OpenedTab)
+			)
+		)
+	);
+
+	// Register the UI COMMANDS and map them to our functions
+	MapActions();
+
+	const bool bCreateDefaultStandaloneMenu = true;
+	const bool bCreateDefaultToolbar = true;
+	// Need editing object to not be null
+	UTranslationUnit* EditingObject;
+	if (DataManager->GetAllTranslationsArray().Num() > 0 && DataManager->GetAllTranslationsArray()[0] != NULL)
+	{
+		EditingObject = DataManager->GetAllTranslationsArray()[0];
+	}
+	else
+	{
+		EditingObject = NewObject<UTranslationUnit>();
+	}
+	FAssetEditorToolkit::InitAssetEditor(Mode, InitToolkitHost, FTranslationEditorModule::TranslationEditorAppIdentifier, StandaloneDefaultLayout, bCreateDefaultStandaloneMenu, bCreateDefaultToolbar, EditingObject);
+	
+	FTranslationEditorModule& TranslationEditorModule = FModuleManager::LoadModuleChecked<FTranslationEditorModule>( "TranslationEditor" );
+	AddMenuExtender(TranslationEditorModule.GetMenuExtensibilityManager()->GetAllExtenders(GetToolkitCommands(), GetEditingObjects()));
+	TSharedPtr<FExtender> MenuExtender = MakeShareable(new FExtender);
+	FTranslationEditorMenu::SetupTranslationEditorMenu( MenuExtender, *this );
+	AddMenuExtender(MenuExtender);
+
+	AddToolbarExtender(TranslationEditorModule.GetToolbarExtensibilityManager()->GetAllExtenders(GetToolkitCommands(), GetEditingObjects()));
+	TSharedPtr<FExtender> ToolbarExtender = MakeShareable(new FExtender);
+	FTranslationEditorMenu::SetupTranslationEditorToolbar( ToolbarExtender, *this );
+	AddToolbarExtender(ToolbarExtender);
+
+	RegenerateMenusAndToolbars();
+
+	// @todo toolkit world centric editing
+	/*// Setup our tool's layout
+	if( IsWorldCentricAssetEditor() )
+	{
+		const FString TabInitializationPayload(TEXT(""));		// NOTE: Payload not currently used for table properties
+		SpawnToolkitTab( UntranslatedTabId, TabInitializationPayload, EToolkitTabSpot::Details );
+	}*/
+
+	// NOTE: Could fill in asset editor commands here!
+}
+
+FName FTranslationEditor::GetToolkitFName() const
+{
+	return FName("TranslationEditor");
+}
+
+FText FTranslationEditor::GetBaseToolkitName() const
+{
+	return LOCTEXT( "AppLabel", "Translation Editor" );
+}
+
+FText FTranslationEditor::GetToolkitName() const
+{
+	const UObject* EditingObject = GetEditingObject();
+	check (EditingObject != NULL);
+
+	FFormatNamedArguments Args;
+	Args.Add(TEXT("Language"), FText::FromString(FPaths::GetBaseFilename(FPaths::GetPath(ArchiveFilePath))));
+	Args.Add(TEXT("ProjectName"), FText::FromString(FPaths::GetBaseFilename(ManifestFilePath)));
+	Args.Add( TEXT("ToolkitName"), GetBaseToolkitName() );
+	return FText::Format( LOCTEXT("TranslationEditorAppLabel", "{Language} - {ProjectName} - {ToolkitName}"), Args );
+}
+
+FText FTranslationEditor::GetToolkitToolTipText() const
+{
+	const UObject* EditingObject = GetEditingObject();
+
+	check (EditingObject != NULL);
+
+	FFormatNamedArguments Args;
+	Args.Add(TEXT("Language"), FText::FromString(FPaths::GetBaseFilename(FPaths::GetPath(ArchiveFilePath))));
+	Args.Add(TEXT("ProjectName"), FText::FromString(FPaths::GetBaseFilename(ManifestFilePath)));
+	Args.Add( TEXT("ToolkitName"), GetBaseToolkitName() );
+	return FText::Format( LOCTEXT("TranslationEditorAppToolTip", "{Language} - {ProjectName} - {ToolkitName}"), Args );
+}
+
+FString FTranslationEditor::GetWorldCentricTabPrefix() const
+{
+	return LOCTEXT("WorldCentricTabPrefix", "Translation ").ToString();
+}
+
+FLinearColor FTranslationEditor::GetWorldCentricTabColorScale() const
+{
+	return FLinearColor( 0.0f, 0.0f, 0.2f, 0.5f );
+}
+FText FTranslationEditor::GetTabLabel_Untranslated() const
+{
+	const int32 UntranslatedNum = DataManager->GetUntranslatedArray().Num();
+	return UntranslatedNum > 0 ?
+		FText::Format(LOCTEXT("UntranslatedTabTitleNumbered", "Untranslated ({0})"), UntranslatedNum) :
+		LOCTEXT("UntranslatedTabTitle", "Untranslated");
+}
+
+TSharedRef<SDockTab> FTranslationEditor::SpawnTab_Untranslated( const FSpawnTabArgs& Args )
+{
+	check( Args.GetTabId().TabType == UntranslatedTabId );
+
+	FPropertyEditorModule& PropertyEditorModule = FModuleManager::GetModuleChecked<FPropertyEditorModule>( "PropertyEditor" );
+
+	FProperty* SourceProperty = FindFProperty<FProperty>( UTranslationUnit::StaticClass(), "Source");
+	FProperty* TranslationProperty = FindFProperty<FProperty>( UTranslationUnit::StaticClass(), "Translation");
+
+	// create empty property table
+	UntranslatedPropertyTable = PropertyEditorModule.CreatePropertyTable();
+	UntranslatedPropertyTable->SetIsUserAllowedToChangeRoot( false );
+	UntranslatedPropertyTable->SetOrientation( EPropertyTableOrientation::AlignPropertiesInColumns );
+	UntranslatedPropertyTable->SetShowRowHeader( true );
+	UntranslatedPropertyTable->SetShowObjectName( false );
+	UntranslatedPropertyTable->OnSelectionChanged()->AddSP( this, &FTranslationEditor::UpdateUntranslatedSelection );
+
+	// we want to customize some columns
+	TArray< TSharedRef<class IPropertyTableCustomColumn>> CustomColumns;
+	SourceColumn->AddSupportedProperty(SourceProperty);
+	TranslationColumn->AddSupportedProperty(TranslationProperty);
+	CustomColumns.Add( SourceColumn );
+	CustomColumns.Add(TranslationColumn);
+
+	UntranslatedPropertyTable->SetObjects((TArray<UObject*>&)DataManager->GetUntranslatedArray());
+
+	// Add the columns we want to display
+	UntranslatedPropertyTable->AddColumn((TWeakFieldPtr<FProperty>)SourceProperty);
+	UntranslatedPropertyTable->AddColumn((TWeakFieldPtr<FProperty>)TranslationProperty);
+
+	// Freeze columns, don't want user to remove them
+	TArray<TSharedRef<IPropertyTableColumn>> Columns = UntranslatedPropertyTable->GetColumns();
+	for (TSharedRef<IPropertyTableColumn> Column : Columns)
+	{
+		Column->SetFrozen(true);
+	}
+
+	UntranslatedPropertyTableWidgetHandle = PropertyEditorModule.CreatePropertyTableWidgetHandle( UntranslatedPropertyTable.ToSharedRef(), CustomColumns);
+	TSharedRef<SWidget> PropertyTableWidget = UntranslatedPropertyTableWidgetHandle->GetWidget();
+
+	
+
+	TSharedRef<SDockTab> NewDockTab = SNew(SDockTab)
+		//.Icon( FAppStyle::GetBrush("TranslationEditor.Tabs.Properties") )
+		.Label(this, &FTranslationEditor::GetTabLabel_Untranslated)
+		.TabColorScale( GetTabColorScale() )
+		[
+			SNew(SBorder)
+			.BorderImage( FAppStyle::GetBrush("ToolPanel.GroupBorder") )
+			.Padding(0.0f)
+			[
+				PropertyTableWidget
+			]
+		];
+
+	UntranslatedTab = NewDockTab;
+
+	return NewDockTab;
+}
+
+FText FTranslationEditor::GetTabLabel_Review() const
+{
+	const int32 ReviewNum = DataManager->GetReviewArray().Num();
+	return ReviewNum > 0 ?
+		FText::Format(LOCTEXT("ReviewTabTitleNumbered", "Needs Review ({0})"), ReviewNum) :
+		LOCTEXT("ReviewTabTitle", "Needs Review");
+}
+
+TSharedRef<SDockTab> FTranslationEditor::SpawnTab_Review( const FSpawnTabArgs& Args )
+{
+	check( Args.GetTabId().TabType == ReviewTabId );
+
+	FProperty* SourceProperty = FindFProperty<FProperty>( UTranslationUnit::StaticClass(), "Source");
+	FProperty* TranslationProperty = FindFProperty<FProperty>( UTranslationUnit::StaticClass(), "Translation");
+
+	FPropertyEditorModule& PropertyEditorModule = FModuleManager::GetModuleChecked<FPropertyEditorModule>( "PropertyEditor" );
+
+	// create empty property table
+	ReviewPropertyTable = PropertyEditorModule.CreatePropertyTable();
+	ReviewPropertyTable->SetIsUserAllowedToChangeRoot( false );
+	ReviewPropertyTable->SetOrientation( EPropertyTableOrientation::AlignPropertiesInColumns );
+	ReviewPropertyTable->SetShowRowHeader( true );
+	ReviewPropertyTable->SetShowObjectName( false );
+	ReviewPropertyTable->OnSelectionChanged()->AddSP( this, &FTranslationEditor::UpdateNeedsReviewSelection );
+
+	// we want to customize some columns
+	TArray< TSharedRef< class IPropertyTableCustomColumn > > CustomColumns;
+	SourceColumn->AddSupportedProperty(SourceProperty);
+	TranslationColumn->AddSupportedProperty(TranslationProperty);
+	CustomColumns.Add( SourceColumn );
+	CustomColumns.Add( TranslationColumn );
+
+	ReviewPropertyTable->SetObjects((TArray<UObject*>&)DataManager->GetReviewArray());
+
+	// Add the columns we want to display
+	ReviewPropertyTable->AddColumn((TWeakFieldPtr<FProperty>)FindFProperty<FProperty>( UTranslationUnit::StaticClass(), "Source"));
+	ReviewPropertyTable->AddColumn((TWeakFieldPtr<FProperty>)FindFProperty<FProperty>( UTranslationUnit::StaticClass(), "Translation"));
+	ReviewPropertyTable->AddColumn((TWeakFieldPtr<FProperty>)FindFProperty<FProperty>( UTranslationUnit::StaticClass(), "HasBeenReviewed"));
+
+	TArray<TSharedRef<IPropertyTableColumn>> Columns = ReviewPropertyTable->GetColumns();
+	for (TSharedRef<IPropertyTableColumn> Column : Columns)
+	{
+		FString ColumnId = Column->GetId().ToString();
+		if (ColumnId == "HasBeenReviewed")
+		{
+			Column->SetWidth(120);
+			Column->SetSizeMode(EPropertyTableColumnSizeMode::Fixed);
+		}
+		// Freeze columns, don't want user to remove them
+		Column->SetFrozen(true);
+	}
+
+	ReviewPropertyTableWidgetHandle = PropertyEditorModule.CreatePropertyTableWidgetHandle( ReviewPropertyTable.ToSharedRef(), CustomColumns);
+	TSharedRef<SWidget> PropertyTableWidget = ReviewPropertyTableWidgetHandle->GetWidget();
+
+	
+
+	TSharedRef<SDockTab> NewDockTab = SNew(SDockTab)
+		//.Icon( FAppStyle::GetBrush("TranslationEditor.Tabs.Properties") )
+		.Label(this, &FTranslationEditor::GetTabLabel_Review)
+		.TabColorScale( GetTabColorScale() )
+		[
+			SNew(SBorder)
+			.BorderImage( FAppStyle::GetBrush("ToolPanel.GroupBorder") )
+			.Padding(0.0f)
+			[
+				PropertyTableWidget
+			]
+		];
+
+	ReviewTab = NewDockTab;
+
+	return NewDockTab;
+}
+
+FText FTranslationEditor::GetTabLabel_Completed() const
+{
+	const int32 CompletedNum = DataManager->GetCompleteArray().Num();
+	return CompletedNum > 0 ?
+		FText::Format(LOCTEXT("CompletedTabTitleNumbered", "Completed ({0})"), CompletedNum) :
+		LOCTEXT("CompletedTabTitle", "Completed");
+}
+
+TSharedRef<SDockTab> FTranslationEditor::SpawnTab_Completed( const FSpawnTabArgs& Args )
+{
+	check( Args.GetTabId().TabType == CompletedTabId );
+
+	FProperty* SourceProperty = FindFProperty<FProperty>( UTranslationUnit::StaticClass(), "Source");
+	FProperty* TranslationProperty = FindFProperty<FProperty>( UTranslationUnit::StaticClass(), "Translation");
+
+	FPropertyEditorModule& PropertyEditorModule = FModuleManager::GetModuleChecked<FPropertyEditorModule>( "PropertyEditor" );
+
+	// create empty property table
+	CompletedPropertyTable = PropertyEditorModule.CreatePropertyTable();
+	CompletedPropertyTable->SetIsUserAllowedToChangeRoot( false );
+	CompletedPropertyTable->SetOrientation( EPropertyTableOrientation::AlignPropertiesInColumns );
+	CompletedPropertyTable->SetShowRowHeader( true );
+	CompletedPropertyTable->SetShowObjectName( false );
+	CompletedPropertyTable->OnSelectionChanged()->AddSP( this, &FTranslationEditor::UpdateCompletedSelection );
+
+	// we want to customize some columns
+	TArray< TSharedRef< class IPropertyTableCustomColumn > > CustomColumns;
+	SourceColumn->AddSupportedProperty(SourceProperty);
+	TranslationColumn->AddSupportedProperty(TranslationProperty);
+	CustomColumns.Add( SourceColumn );
+	CustomColumns.Add( TranslationColumn );
+
+	CompletedPropertyTable->SetObjects((TArray<UObject*>&)DataManager->GetCompleteArray());
+
+	// Add the columns we want to display
+	CompletedPropertyTable->AddColumn((TWeakFieldPtr<FProperty>)FindFProperty<FProperty>( UTranslationUnit::StaticClass(), "Source"));
+	CompletedPropertyTable->AddColumn((TWeakFieldPtr<FProperty>)FindFProperty<FProperty>( UTranslationUnit::StaticClass(), "Translation"));
+
+	// Freeze columns, don't want user to remove them
+	TArray<TSharedRef<IPropertyTableColumn>> Columns = CompletedPropertyTable->GetColumns();
+	for (TSharedRef<IPropertyTableColumn> Column : Columns)
+	{
+		Column->SetFrozen(true);
+	}
+
+	CompletedPropertyTableWidgetHandle = PropertyEditorModule.CreatePropertyTableWidgetHandle( CompletedPropertyTable.ToSharedRef(), CustomColumns);
+	TSharedRef<SWidget> PropertyTableWidget = CompletedPropertyTableWidgetHandle->GetWidget();
+
+	
+
+	TSharedRef<SDockTab> NewDockTab = SNew(SDockTab)
+		//.Icon( FAppStyle::GetBrush("TranslationEditor.Tabs.Properties") )
+		.Label(this, &FTranslationEditor::GetTabLabel_Completed)
+		.TabColorScale( GetTabColorScale() )
+		[
+			SNew(SBorder)
+			.BorderImage( FAppStyle::GetBrush("ToolPanel.GroupBorder") )
+			.Padding(0.0f)
+			[
+				PropertyTableWidget
+			]
+		];
+
+	CompletedTab = NewDockTab;
+
+	return NewDockTab;
+}
+
+TSharedRef<SDockTab> FTranslationEditor::SpawnTab_Search(const FSpawnTabArgs& Args)
+{
+	check(Args.GetTabId().TabType == SearchTabId);
+
+	FProperty* SourceProperty = FindFProperty<FProperty>(UTranslationUnit::StaticClass(), "Source");
+	FProperty* TranslationProperty = FindFProperty<FProperty>(UTranslationUnit::StaticClass(), "Translation");
+
+	FPropertyEditorModule& PropertyEditorModule = FModuleManager::GetModuleChecked<FPropertyEditorModule>("PropertyEditor");
+
+	// create empty property table
+	SearchPropertyTable = PropertyEditorModule.CreatePropertyTable();
+	SearchPropertyTable->SetIsUserAllowedToChangeRoot(false);
+	SearchPropertyTable->SetOrientation(EPropertyTableOrientation::AlignPropertiesInColumns);
+	SearchPropertyTable->SetShowRowHeader(true);
+	SearchPropertyTable->SetShowObjectName(false);
+	SearchPropertyTable->OnSelectionChanged()->AddSP(this, &FTranslationEditor::UpdateSearchSelection);
+
+	// we want to customize some columns
+	TArray< TSharedRef< class IPropertyTableCustomColumn > > CustomColumns;
+	SourceColumn->AddSupportedProperty(SourceProperty);
+	TranslationColumn->AddSupportedProperty(TranslationProperty);
+	CustomColumns.Add(SourceColumn);
+	CustomColumns.Add(TranslationColumn);
+
+	SearchPropertyTable->SetObjects((TArray<UObject*>&)DataManager->GetSearchResultsArray());
+
+	// Add the columns we want to display
+	SearchPropertyTable->AddColumn((TWeakFieldPtr<FProperty>)FindFProperty<FProperty>(UTranslationUnit::StaticClass(), "Source"));
+	SearchPropertyTable->AddColumn((TWeakFieldPtr<FProperty>)FindFProperty<FProperty>(UTranslationUnit::StaticClass(), "Translation"));
+
+	// Freeze columns, don't want user to remove them
+	TArray<TSharedRef<IPropertyTableColumn>> Columns = SearchPropertyTable->GetColumns();
+	for (TSharedRef<IPropertyTableColumn> Column : Columns)
+	{
+		Column->SetFrozen(true);
+	}
+
+	SearchPropertyTableWidgetHandle = PropertyEditorModule.CreatePropertyTableWidgetHandle(SearchPropertyTable.ToSharedRef(), CustomColumns);
+	TSharedRef<SWidget> PropertyTableWidget = SearchPropertyTableWidgetHandle->GetWidget();
+
+	TSharedRef<SDockTab> NewDockTab = SNew(SDockTab)
+		//.Icon(FAppStyle::GetBrush("TranslationEditor.Tabs.Properties"))
+		.Label(LOCTEXT("SearchTabTitle", "Search"))
+		.TabColorScale(GetTabColorScale())
+		[
+			SNew(SVerticalBox)
+			+ SVerticalBox::Slot()
+			.HAlign(HAlign_Fill)
+			.VAlign(VAlign_Top)
+			.AutoHeight()
+			.Padding(0.0f, 0.0f, 0.0f, 4.0f)
+			[
+				SAssignNew(SearchBox, SSearchBox)
+				.HintText(LOCTEXT("FilterSearch", "Search..."))
+				.ToolTipText(LOCTEXT("FilterSearchHint", "Type here to search"))
+				.OnTextChanged(this, &FTranslationEditor::OnFilterTextChanged)
+				.OnTextCommitted(this, &FTranslationEditor::OnFilterTextCommitted)
+			]
+			+ SVerticalBox::Slot()
+				.HAlign(HAlign_Fill)
+				.VAlign(VAlign_Top)
+				.FillHeight(10.f)
+			[
+				SNew(SBorder)
+				.BorderImage(FAppStyle::GetBrush("ToolPanel.GroupBorder"))
+				.Padding(0.0f)
+				.VAlign(VAlign_Top)
+				[
+					PropertyTableWidget
+				]
+			]
+		];
+
+	SearchTab = NewDockTab;
+
+	return NewDockTab;
+}
+
+TSharedRef<SDockTab> FTranslationEditor::SpawnTab_ChangedOnImport(const FSpawnTabArgs& Args)
+{
+	check(Args.GetTabId().TabType == ChangedOnImportTabId);
+
+	FProperty* SourceProperty = FindFProperty<FProperty>(UTranslationUnit::StaticClass(), "Source");
+	FProperty* TranslationBeforeImportProperty = FindFProperty<FProperty>(UTranslationUnit::StaticClass(), "TranslationBeforeImport");
+	FProperty* TranslationProperty = FindFProperty<FProperty>(UTranslationUnit::StaticClass(), "Translation");
+
+	FPropertyEditorModule& PropertyEditorModule = FModuleManager::GetModuleChecked<FPropertyEditorModule>("PropertyEditor");
+
+	// create empty property table
+	ChangedOnImportPropertyTable = PropertyEditorModule.CreatePropertyTable();
+	ChangedOnImportPropertyTable->SetIsUserAllowedToChangeRoot(false);
+	ChangedOnImportPropertyTable->SetOrientation(EPropertyTableOrientation::AlignPropertiesInColumns);
+	ChangedOnImportPropertyTable->SetShowRowHeader(true);
+	ChangedOnImportPropertyTable->SetShowObjectName(false);
+	ChangedOnImportPropertyTable->OnSelectionChanged()->AddSP(this, &FTranslationEditor::UpdateSearchSelection);
+
+	// we want to customize some columns
+	TArray< TSharedRef< class IPropertyTableCustomColumn > > CustomColumns;
+	SourceColumn->AddSupportedProperty(SourceProperty);
+	TranslationColumn->AddSupportedProperty(TranslationProperty);
+	CustomColumns.Add(SourceColumn);
+	CustomColumns.Add(TranslationColumn);
+
+	ChangedOnImportPropertyTable->SetObjects((TArray<UObject*>&)DataManager->GetSearchResultsArray());
+
+	// Add the columns we want to display
+	ChangedOnImportPropertyTable->AddColumn((TWeakFieldPtr<FProperty>)FindFProperty<FProperty>(UTranslationUnit::StaticClass(), "Source"));
+	ChangedOnImportPropertyTable->AddColumn((TWeakFieldPtr<FProperty>)FindFProperty<FProperty>(UTranslationUnit::StaticClass(), "TranslationBeforeImport"));
+	ChangedOnImportPropertyTable->AddColumn((TWeakFieldPtr<FProperty>)FindFProperty<FProperty>(UTranslationUnit::StaticClass(), "Translation"));
+
+	// Freeze columns, don't want user to remove them
+	TArray<TSharedRef<IPropertyTableColumn>> Columns = ChangedOnImportPropertyTable->GetColumns();
+	for (TSharedRef<IPropertyTableColumn> Column : Columns)
+	{
+		Column->SetFrozen(true);
+	}
+
+	SearchPropertyTableWidgetHandle = PropertyEditorModule.CreatePropertyTableWidgetHandle(ChangedOnImportPropertyTable.ToSharedRef(), CustomColumns);
+	TSharedRef<SWidget> PropertyTableWidget = SearchPropertyTableWidgetHandle->GetWidget();
+
+	TSharedRef<SDockTab> NewDockTab = SNew(SDockTab)
+		//.Icon(FAppStyle::GetBrush("TranslationEditor.Tabs.Properties"))
+		.Label(LOCTEXT("ChangedOnImportTabTitle", "Changed on Import"))
+		.TabColorScale(GetTabColorScale())
+		[
+			SNew(SBorder)
+			.BorderImage(FAppStyle::GetBrush("ToolPanel.GroupBorder"))
+			.Padding(0.0f)
+			[
+				PropertyTableWidget
+			]
+		];
+
+	ChangedOnImportTab = NewDockTab;
+
+	return NewDockTab;
+}
+
+TSharedRef<SDockTab> FTranslationEditor::SpawnTab_Preview( const FSpawnTabArgs& Args )
+{
+	check( Args.GetTabId().TabType == PreviewTabId );
+
+	TSharedRef<SDockTab> NewDockTab = SNew(SDockTab)
+		//.Icon( FAppStyle::GetBrush("TranslationEditor.Tabs.Properties") )
+		.Label( LOCTEXT("PreviewTabTitle", "Preview") )
+		.TabColorScale( GetTabColorScale() )
+		[
+			SNew(SBorder)
+			.BorderImage( FAppStyle::GetBrush("ToolPanel.GroupBorder") )
+			.Padding(0.0f)
+			[
+				SNew(SHorizontalBox)
+				+SHorizontalBox::Slot()
+				.HAlign(HAlign_Center)
+				.VAlign(VAlign_Center)
+				[
+					PreviewTextBlock
+				]
+			]
+		];
+	
+	return NewDockTab;
+}
+
+TSharedRef<SDockTab> FTranslationEditor::SpawnTab_Context( const FSpawnTabArgs& Args )
+{
+	check( Args.GetTabId().TabType == ContextTabId );
+
+	FPropertyEditorModule& PropertyEditorModule = FModuleManager::GetModuleChecked<FPropertyEditorModule>( "PropertyEditor" );
+
+	// create empty property table
+	ContextPropertyTable = PropertyEditorModule.CreatePropertyTable();
+	ContextPropertyTable->SetIsUserAllowedToChangeRoot( false );
+	ContextPropertyTable->SetOrientation( EPropertyTableOrientation::AlignPropertiesInColumns );
+	ContextPropertyTable->SetShowRowHeader( true );
+	ContextPropertyTable->SetShowObjectName( false );
+	ContextPropertyTable->OnSelectionChanged()->AddSP( this, &FTranslationEditor::UpdateContextSelection );
+
+	if (DataManager->GetAllTranslationsArray().Num() > 0)
+	{
+		TArray<UObject*> Objects;
+		Objects.Add(DataManager->GetAllTranslationsArray()[0]);
+		ContextPropertyTable->SetObjects(Objects);
+	}
+
+	// Build the Path to the data we want to show
+	FProperty* ContextProp = FindFProperty<FProperty>( UTranslationUnit::StaticClass(), "Contexts" );
+	FPropertyInfo ContextPropInfo;
+	ContextPropInfo.Property = ContextProp;
+	ContextPropInfo.ArrayIndex = INDEX_NONE;
+	TSharedRef<FPropertyPath> Path = FPropertyPath::CreateEmpty();
+	Path = Path->ExtendPath(ContextPropInfo);
+	ContextPropertyTable->SetRootPath(Path);
+
+	// Add the columns we want to display
+	ContextPropertyTable->AddColumn((TWeakFieldPtr<FProperty>)FindFProperty<FProperty>( FTranslationContextInfo::StaticStruct(), "Key"));
+	ContextPropertyTable->AddColumn((TWeakFieldPtr<FProperty>)FindFProperty<FProperty>( FTranslationContextInfo::StaticStruct(), "Context"));
+
+	// Freeze columns, don't want user to remove them
+	TArray<TSharedRef<IPropertyTableColumn>> Columns = ContextPropertyTable->GetColumns();
+	for (TSharedRef<IPropertyTableColumn> Column : Columns)
+	{
+		Column->SetFrozen(true);
+	}
+
+	ContextPropertyTableWidgetHandle = PropertyEditorModule.CreatePropertyTableWidgetHandle( ContextPropertyTable.ToSharedRef() );
+	TSharedRef<SWidget> PropertyTableWidget = ContextPropertyTableWidgetHandle->GetWidget();
+
+	TSharedRef<SDockTab> NewDockTab = SNew(SDockTab)
+		//.Icon( FAppStyle::GetBrush("TranslationEditor.Tabs.Properties") )
+		.Label( LOCTEXT("ContextTabTitle", "Context") )
+		.TabColorScale( GetTabColorScale() )
+		[
+			SNew(SBorder)
+			.BorderImage( FAppStyle::GetBrush("ToolPanel.GroupBorder") )
+			.Padding(0.0f)
+			[
+				SNew(SVerticalBox)
+				+SVerticalBox::Slot()
+				.HAlign(HAlign_Left)
+				.VAlign(VAlign_Center)
+				.AutoHeight()
+				[
+					SNew(SBorder)
+					.BorderImage(FAppStyle::GetBrush("ToolPanel.GroupBorder"))
+					.Padding(5.0f)
+					[
+						NamespaceTextBlock
+					]
+				]
+				+SVerticalBox::Slot()
+				.HAlign(HAlign_Fill)
+				.VAlign(VAlign_Fill)
+				[
+					PropertyTableWidget
+				]
+			]
+		];
+	
+	return NewDockTab;
+}
+
+TSharedRef<SDockTab> FTranslationEditor::SpawnTab_History(const FSpawnTabArgs& Args)
+{
+	check(Args.GetTabId().TabType == HistoryTabId);
+
+	FPropertyEditorModule& PropertyEditorModule = FModuleManager::GetModuleChecked<FPropertyEditorModule>("PropertyEditor");
+
+	FProperty* SourceProperty = FindFProperty<FProperty>(FTranslationChange::StaticStruct(), "Source");
+	FProperty* TranslationProperty = FindFProperty<FProperty>(FTranslationChange::StaticStruct(), "Translation");
+
+	// create empty property table
+	HistoryPropertyTable = PropertyEditorModule.CreatePropertyTable();
+	HistoryPropertyTable->SetIsUserAllowedToChangeRoot(false);
+	HistoryPropertyTable->SetOrientation(EPropertyTableOrientation::AlignPropertiesInColumns);
+	HistoryPropertyTable->SetShowRowHeader(true);
+	HistoryPropertyTable->SetShowObjectName(false);
+
+	// we want to customize some columns
+	TArray< TSharedRef<class IPropertyTableCustomColumn>> CustomColumns;
+	SourceColumn->AddSupportedProperty(SourceProperty);
+	TranslationColumn->AddSupportedProperty(TranslationProperty);
+	CustomColumns.Add(SourceColumn);
+	CustomColumns.Add(TranslationColumn);
+
+	if (DataManager->GetAllTranslationsArray().Num() > 0)
+	{
+		TArray<UObject*> Objects;
+		Objects.Add(DataManager->GetAllTranslationsArray()[0]);
+		HistoryPropertyTable->SetObjects(Objects);
+	}
+
+	// Build the Path to the data we want to show
+	TSharedRef<FPropertyPath> Path = FPropertyPath::CreateEmpty();
+	FArrayProperty* ContextsProp = FindFProperty<FArrayProperty>(UTranslationUnit::StaticClass(), "Contexts");
+	Path = Path->ExtendPath(FPropertyPath::Create(ContextsProp));
+	FPropertyInfo ContextsPropInfo;
+	ContextsPropInfo.Property = ContextsProp->Inner;
+	ContextsPropInfo.ArrayIndex = 0;
+	Path = Path->ExtendPath(ContextsPropInfo);
+
+	FProperty* ChangesProp = FindFProperty<FProperty>(FTranslationContextInfo::StaticStruct(), "Changes");
+	FPropertyInfo ChangesPropInfo;
+	ChangesPropInfo.Property = ChangesProp;
+	ChangesPropInfo.ArrayIndex = INDEX_NONE;
+	Path = Path->ExtendPath(ChangesPropInfo);
+	HistoryPropertyTable->SetRootPath(Path);
+
+	// Add the columns we want to display
+	HistoryPropertyTable->AddColumn((TWeakFieldPtr<FProperty>)FindFProperty<FProperty>(FTranslationChange::StaticStruct(), "Version"));
+	HistoryPropertyTable->AddColumn((TWeakFieldPtr<FProperty>)FindFProperty<FProperty>(FTranslationChange::StaticStruct(), "DateAndTime"));
+	HistoryPropertyTable->AddColumn((TWeakFieldPtr<FProperty>)SourceProperty);
+	HistoryPropertyTable->AddColumn((TWeakFieldPtr<FProperty>)TranslationProperty);
+
+	// Freeze columns, don't want user to remove them
+	TArray<TSharedRef<IPropertyTableColumn>> Columns = HistoryPropertyTable->GetColumns();
+	for (TSharedRef<IPropertyTableColumn> Column : Columns)
+	{
+		Column->SetFrozen(true);
+	}
+
+	HistoryPropertyTableWidgetHandle = PropertyEditorModule.CreatePropertyTableWidgetHandle(HistoryPropertyTable.ToSharedRef(), CustomColumns);
+	TSharedRef<SWidget> PropertyTableWidget = HistoryPropertyTableWidgetHandle->GetWidget();
+
+	TSharedRef<SDockTab> NewDockTab = SNew(SDockTab)
+		//.Icon(FAppStyle::GetBrush("TranslationEditor.Tabs.Properties"))
+		.Label(LOCTEXT("HistoryTabTitle", "History"))
+		.TabColorScale(GetTabColorScale())
+		[
+			SNew(SBorder)
+			.BorderImage(FAppStyle::GetBrush("ToolPanel.GroupBorder"))
+			.Padding(0.0f)
+			[
+				SNew(SVerticalBox)
+				+ SVerticalBox::Slot()
+				.HAlign(HAlign_Left)
+				.VAlign(VAlign_Center)
+				.AutoHeight()
+				[
+					SNew(SBorder)
+					.BorderImage(FAppStyle::GetBrush("ToolPanel.GroupBorder"))
+					.Padding(5.0f)
+					[
+						SNew(SButton)
+						.HAlign(HAlign_Center)
+						.VAlign(VAlign_Center)
+						.OnClicked(FOnClicked::CreateSP(this, &FTranslationEditor::OnGetHistoryButtonClicked))
+						[
+							SNew(STextBlock)
+							.Text((LOCTEXT("GetHistoryButton", "Get History...")))
+						]
+					]
+				]
+				+ SVerticalBox::Slot()
+				.HAlign(HAlign_Fill)
+				.VAlign(VAlign_Fill)
+				[
+					PropertyTableWidget
+				]
+			]
+		];
+	
+	return NewDockTab;
+}
+
+void FTranslationEditor::MapActions()
+{
+	FTranslationEditorCommands::Register();
+
+	ToolkitCommands->MapAction( FTranslationEditorCommands::Get().ChangeSourceFont,
+		FExecuteAction::CreateSP(this, &FTranslationEditor::ChangeSourceFont),
+		FCanExecuteAction());
+
+	ToolkitCommands->MapAction( FTranslationEditorCommands::Get().ChangeTranslationTargetFont,
+		FExecuteAction::CreateSP(this, &FTranslationEditor::ChangeTranslationTargetFont),
+		FCanExecuteAction());
+
+	ToolkitCommands->MapAction( FTranslationEditorCommands::Get().SaveTranslations,
+		FExecuteAction::CreateSP(this, &FTranslationEditor::SaveAsset_Execute),
+		FCanExecuteAction());
+
+	ToolkitCommands->MapAction(FTranslationEditorCommands::Get().PreviewAllTranslationsInEditor,
+		FExecuteAction::CreateSP(this, &FTranslationEditor::PreviewAllTranslationsInEditor_Execute),
+		FCanExecuteAction());
+
+	ToolkitCommands->MapAction(FTranslationEditorCommands::Get().ImportLatestFromLocalizationService,
+		FExecuteAction::CreateSP(this, &FTranslationEditor::ImportLatestFromLocalizationService_Execute),
+		FCanExecuteAction());
+
+	ToolkitCommands->MapAction(FTranslationEditorCommands::Get().ExportToPortableObjectFormat,
+		FExecuteAction::CreateSP(this, &FTranslationEditor::ExportToPortableObjectFormat_Execute),
+		FCanExecuteAction());
+
+	ToolkitCommands->MapAction(FTranslationEditorCommands::Get().ImportFromPortableObjectFormat,
+		FExecuteAction::CreateSP(this, &FTranslationEditor::ImportFromPortableObjectFormat_Execute),
+		FCanExecuteAction());
+
+	ToolkitCommands->MapAction(FTranslationEditorCommands::Get().OpenSearchTab,
+		FExecuteAction::CreateSP(this, &FTranslationEditor::OpenSearchTab_Execute),
+		FCanExecuteAction());
+
+	ToolkitCommands->MapAction(FTranslationEditorCommands::Get().OpenTranslationPicker,
+		FExecuteAction::CreateStatic(&ITranslationEditor::OpenTranslationPicker),
+		FCanExecuteAction());
+}
+
+void FTranslationEditor::ChangeSourceFont()
+{
+	// Use path from current font
+	FString DefaultFile = TranslationEditorUtils::GetFontFilename(SourceFont);
+
+	FString NewFontFilename;
+	bool bOpened = OpenFontPicker(DefaultFile, NewFontFilename);
+
+	if ( bOpened && NewFontFilename != "")
+	{
+		SourceFont = FSlateFontInfo(NewFontFilename, SourceFont.Size);
+		RefreshUI();
+	}
+}
+
+void FTranslationEditor::ChangeTranslationTargetFont()
+{
+	// Use path from current font
+	FString DefaultFile = TranslationEditorUtils::GetFontFilename(TranslationTargetFont);
+
+	FString NewFontFilename;
+	bool bOpened = OpenFontPicker(DefaultFile, NewFontFilename);
+
+	if ( bOpened && NewFontFilename != "")
+	{
+		TranslationTargetFont = FSlateFontInfo(NewFontFilename, TranslationTargetFont.Size);
+		RefreshUI();
+	}
+}
+
+void FTranslationEditor::RefreshUI()
+{
+	// Set the fonts in our custom font columns and text block
+	SourceColumn->SetFont(SourceFont);
+	TranslationColumn->SetFont(TranslationTargetFont);
+	PreviewTextBlock->SetFont(TranslationTargetFont);
+
+	// Refresh our widget displays
+	if (UntranslatedPropertyTableWidgetHandle.IsValid())
+	{
+		UntranslatedPropertyTableWidgetHandle->RequestRefresh();
+	}
+	if (ReviewPropertyTableWidgetHandle.IsValid())	
+	{
+		ReviewPropertyTableWidgetHandle->RequestRefresh();
+	}
+	if (CompletedPropertyTableWidgetHandle.IsValid())
+	{
+		CompletedPropertyTableWidgetHandle->RequestRefresh();
+	}
+	if (ContextPropertyTableWidgetHandle.IsValid())
+	{
+		ContextPropertyTableWidgetHandle->RequestRefresh();
+	}
+	if (HistoryPropertyTableWidgetHandle.IsValid())
+	{
+		HistoryPropertyTableWidgetHandle->RequestRefresh();
+	}
+	if (SearchPropertyTableWidgetHandle.IsValid())
+	{
+		SearchPropertyTableWidgetHandle->RequestRefresh();
+	}
+	if (ChangedOnImportPropertyTableWidgetHandle.IsValid())
+	{
+		ChangedOnImportPropertyTableWidgetHandle->RequestRefresh();
+	}
+}
+
+bool FTranslationEditor::OpenFontPicker( const FString DefaultFile, FString& OutFile )
+{
+	const FString FontFileDescription = LOCTEXT( "FontFileDescription", "Font File" ).ToString();
+	const FString FontFileExtension = TEXT("*.ttf;*.otf");
+	const FString FileTypes = FString::Printf( TEXT("%s (%s)|%s"), *FontFileDescription, *FontFileExtension, *FontFileExtension );
+
+	// Prompt the user for the filenames
+	TArray<FString> OpenFilenames;
+	IDesktopPlatform* DesktopPlatform = FDesktopPlatformModule::Get();
+	bool bOpened = false;
+	if ( DesktopPlatform )
+	{
+		void* ParentWindowWindowHandle = NULL;
+
+		const TSharedPtr<SWindow>& ParentWindow = FSlateApplication::Get().FindWidgetWindow(PreviewTextBlock);
+		if ( ParentWindow.IsValid() && ParentWindow->GetNativeWindow().IsValid() )
+		{
+			ParentWindowWindowHandle = ParentWindow->GetNativeWindow()->GetOSWindowHandle();
+		}
+
+		bOpened = DesktopPlatform->OpenFileDialog(
+			ParentWindowWindowHandle,
+			LOCTEXT("ChooseFontWindowTitle", "Choose Font").ToString(),
+			FPaths::GetPath(DefaultFile),
+			TEXT(""),
+			FileTypes,
+			EFileDialogFlags::None,
+			OpenFilenames
+			);
+	}
+
+	if ( bOpened && OpenFilenames.Num() > 0 )
+	{
+		OutFile = OpenFilenames[0];
+	} 
+	else
+	{
+		OutFile = "";
+	}
+
+	return bOpened;
+}
+
+void FTranslationEditor::UpdateUntranslatedSelection()
+{
+	TSharedPtr<SDockTab> UntranslatedTabSharedPtr = UntranslatedTab.Pin();
+	if (UntranslatedTabSharedPtr.IsValid() && UntranslatedTabSharedPtr->IsForeground() && UntranslatedPropertyTable.IsValid())
+	{
+		TSet<TSharedRef<IPropertyTableRow>> SelectedRows = UntranslatedPropertyTable->GetSelectedRows();
+		UpdateTranslationUnitSelection(SelectedRows);
+	}
+}
+
+void FTranslationEditor::UpdateNeedsReviewSelection()
+{
+	TSharedPtr<SDockTab> ReviewTabSharedPtr = ReviewTab.Pin();
+	if (ReviewTabSharedPtr.IsValid() && ReviewTabSharedPtr->IsForeground() && ReviewPropertyTable.IsValid())
+	{
+		TSet<TSharedRef<IPropertyTableRow>> SelectedRows = ReviewPropertyTable->GetSelectedRows();
+		UpdateTranslationUnitSelection(SelectedRows);
+	}
+}
+
+void FTranslationEditor::UpdateCompletedSelection()
+{
+	TSharedPtr<SDockTab> CompletedTabSharedPtr = CompletedTab.Pin();
+	if (CompletedTabSharedPtr.IsValid() && CompletedTabSharedPtr->IsForeground() && CompletedPropertyTable.IsValid())
+	{
+		TSet<TSharedRef<IPropertyTableRow>> SelectedRows = CompletedPropertyTable->GetSelectedRows();
+		UpdateTranslationUnitSelection(SelectedRows);
+	}
+}
+
+void FTranslationEditor::UpdateSearchSelection()
+{
+	TSharedPtr<SDockTab> SearchTabSharedPtr = SearchTab.Pin();
+	if (SearchTabSharedPtr.IsValid() && SearchTabSharedPtr->IsForeground() && SearchPropertyTable.IsValid())
+	{
+		TSet<TSharedRef<IPropertyTableRow>> SelectedRows = SearchPropertyTable->GetSelectedRows();
+		UpdateTranslationUnitSelection(SelectedRows);
+	}
+}
+
+void FTranslationEditor::UpdateChangedOnImportSelection()
+{
+	TSharedPtr<SDockTab> ChangedOnImportTabSharedPtr = SearchTab.Pin();
+	if (ChangedOnImportTabSharedPtr.IsValid() && ChangedOnImportTabSharedPtr->IsForeground() && ChangedOnImportPropertyTable.IsValid())
+	{
+		TSet<TSharedRef<IPropertyTableRow>> SelectedRows = ChangedOnImportPropertyTable->GetSelectedRows();
+		UpdateTranslationUnitSelection(SelectedRows);
+	}
+}
+
+void FTranslationEditor::UpdateTranslationUnitSelection(TSet<TSharedRef<IPropertyTableRow>>& SelectedRows)
+{
+	// Can only really handle single selection
+	if (SelectedRows.Num() == 1)
+	{
+		TSharedRef<IPropertyTableRow> SelectedRow = *(SelectedRows.CreateConstIterator());
+		TSharedRef<FPropertyPath> PartialPath = SelectedRow->GetPartialPath();
+
+		TWeakObjectPtr<UObject> UObjectWeakPtr = SelectedRow->GetDataSource()->AsUObject();
+		if (UObjectWeakPtr.IsValid())
+		{
+			UObject* UObjectPtr = UObjectWeakPtr.Get();
+			if (UObjectPtr != NULL)
+			{
+				UTranslationUnit* SelectedTranslationUnit = (UTranslationUnit*)UObjectPtr;
+				if (SelectedTranslationUnit != NULL)
+				{
+					PreviewTextBlock->SetText(FText::AsCultureInvariant(SelectedTranslationUnit->Translation));
+					NamespaceTextBlock->SetText(FText::Format(LOCTEXT("TranslationNamespace", "Namespace: {0}"), FText::FromString(SelectedTranslationUnit->Namespace)));
+
+					// Add the ContextPropertyTable-specific path
+					FArrayProperty* ContextArrayProp = FindFProperty<FArrayProperty>(UTranslationUnit::StaticClass(), "Contexts");
+					FPropertyInfo ContextArrayPropInfo;
+					ContextArrayPropInfo.Property = ContextArrayProp;
+					ContextArrayPropInfo.ArrayIndex = INDEX_NONE;
+					TSharedRef<FPropertyPath> ContextPath = FPropertyPath::CreateEmpty();
+					ContextPath = ContextPath->ExtendPath(PartialPath);
+					ContextPath = ContextPath->ExtendPath(ContextArrayPropInfo);
+
+					if (ContextPropertyTable.IsValid())
+					{
+						TArray<UObject*> ObjectArray;
+						ObjectArray.Add(SelectedTranslationUnit);
+						ContextPropertyTable->SetObjects(ObjectArray);
+						ContextPropertyTable->SetRootPath(ContextPath);
+
+						// Need to re-add the columns we want to display
+						ContextPropertyTable->AddColumn((TWeakFieldPtr<FProperty>)FindFProperty<FProperty>(FTranslationContextInfo::StaticStruct(), "Key"));
+						ContextPropertyTable->AddColumn((TWeakFieldPtr<FProperty>)FindFProperty<FProperty>(FTranslationContextInfo::StaticStruct(), "Context"));
+
+						TArray<TSharedRef<IPropertyTableColumn>> Columns = ContextPropertyTable->GetColumns();
+						for (TSharedRef<IPropertyTableColumn> Column : Columns)
+						{
+							Column->SetFrozen(true);
+						}
+
+						TSharedPtr<IPropertyTableCell> ContextToSelectPtr = ContextPropertyTable->GetFirstCellInTable();
+
+						if (ContextToSelectPtr.IsValid())
+						{
+							TSet<TSharedRef<IPropertyTableCell>> CellsToSelect;
+							CellsToSelect.Add(ContextToSelectPtr.ToSharedRef());
+							ContextPropertyTable->SetSelectedCells(CellsToSelect);
+						}
+					}
+				}
+			}
+		}
+	}
+}
+
+void FTranslationEditor::SaveAsset_Execute()
+{
+	// Doesn't call parent SaveAsset_Execute, only need to tell data manager to write data
+	DataManager->WriteTranslationData();
+}
+
+void FTranslationEditor::UpdateContextSelection()
+{
+	if (ContextPropertyTable.IsValid())
+	{
+		TSet<TSharedRef<IPropertyTableRow>> SelectedRows = ContextPropertyTable->GetSelectedRows();
+		TSharedRef<FPropertyPath> InitialPath = ContextPropertyTable->GetRootPath();
+		FProperty* PropertyToFind = InitialPath->GetRootProperty().Property.Get();
+
+		// Can only really handle single selection
+		if (SelectedRows.Num() == 1)
+		{
+			TSharedRef<IPropertyTableRow> SelectedRow = *(SelectedRows.CreateConstIterator());
+			TSharedRef<FPropertyPath> PartialPath = SelectedRow->GetPartialPath();
+
+			TWeakObjectPtr<UObject> UObjectWeakPtr = SelectedRow->GetDataSource()->AsUObject();
+			if (UObjectWeakPtr.IsValid())
+			{
+				UObject* UObjectPtr = UObjectWeakPtr.Get();
+				if (UObjectPtr != NULL)
+				{
+					UTranslationUnit* SelectedTranslationUnit = (UTranslationUnit*)UObjectPtr;
+					if (SelectedTranslationUnit != NULL)
+					{
+						// Index of the leaf most property is the context info index we need
+						FTranslationContextInfo& SelectedContextInfo = SelectedTranslationUnit->Contexts[PartialPath->GetLeafMostProperty().ArrayIndex];
+
+						// If this is a translation unit from the review tab and they select a context, possibly update the selected translation with one from that context
+						// Only change the suggested translation if they haven't yet reviewed it
+						if (SelectedTranslationUnit->HasBeenReviewed == false)
+						{
+							for (int32 ChangeIndex = 0; ChangeIndex < SelectedContextInfo.Changes.Num(); ++ChangeIndex)
+							{
+								// Find most recent, non-empty translation
+								if (!SelectedContextInfo.Changes[ChangeIndex].Translation.IsEmpty() && SelectedTranslationUnit->Translation != SelectedContextInfo.Changes[ChangeIndex].Translation)
+								{
+									SelectedTranslationUnit->Modify();
+									SelectedTranslationUnit->Translation = SelectedContextInfo.Changes[ChangeIndex].Translation;
+									SelectedTranslationUnit->PostEditChange();
+								}
+							}
+						}
+
+
+						// Add the HistoryPropertyTable-specific path
+						TSharedRef<FPropertyPath> HistoryPath = ContextPropertyTable->GetRootPath();
+						FArrayProperty* ContextArrayProp = FindFProperty<FArrayProperty>(UTranslationUnit::StaticClass(), "Contexts");
+						FPropertyInfo ContextPropInfo;
+						ContextPropInfo.Property = ContextArrayProp->Inner;
+						ContextPropInfo.ArrayIndex = PartialPath->GetLeafMostProperty().ArrayIndex;
+						HistoryPath = HistoryPath->ExtendPath(ContextPropInfo);
+						FArrayProperty* ChangesProp = FindFProperty<FArrayProperty>(FTranslationContextInfo::StaticStruct(), "Changes");
+						FPropertyInfo ChangesPropInfo;
+						ChangesPropInfo.Property = ChangesProp;
+						ChangesPropInfo.ArrayIndex = INDEX_NONE;
+						HistoryPath = HistoryPath->ExtendPath(ChangesPropInfo);
+						if (HistoryPropertyTable.IsValid())
+						{
+							TArray<UObject*> ObjectArray;
+							ObjectArray.Add(SelectedTranslationUnit);
+							HistoryPropertyTable->SetObjects(ObjectArray);
+							HistoryPropertyTable->SetRootPath(HistoryPath);
+
+							// Need to re-add the columns we want to display
+							HistoryPropertyTable->AddColumn((TWeakFieldPtr<FProperty>)FindFProperty<FProperty>(FTranslationChange::StaticStruct(), "Version"));
+							HistoryPropertyTable->AddColumn((TWeakFieldPtr<FProperty>)FindFProperty<FProperty>(FTranslationChange::StaticStruct(), "DateAndTime"));
+							HistoryPropertyTable->AddColumn((TWeakFieldPtr<FProperty>)FindFProperty<FProperty>(FTranslationChange::StaticStruct(), "Source"));
+							HistoryPropertyTable->AddColumn((TWeakFieldPtr<FProperty>)FindFProperty<FProperty>(FTranslationChange::StaticStruct(), "Translation"));
+
+							TArray<TSharedRef<IPropertyTableColumn>> Columns = HistoryPropertyTable->GetColumns();
+							for (TSharedRef<IPropertyTableColumn> Column : Columns)
+							{
+								Column->SetFrozen(true);
+							}
+						}
+					}
+				}
+			}
+		}
+	}
+}
+
+void FTranslationEditor::PreviewAllTranslationsInEditor_Execute()
+{
+	DataManager->PreviewAllTranslationsInEditor(AssociatedLocalizationTarget.Get());
+}
+
+void FTranslationEditor::ImportLatestFromLocalizationService_Execute()
+{
+	check(AssociatedLocalizationTarget.IsValid());
+	ILocalizationServiceProvider& Provider = ILocalizationServiceModule::Get().GetProvider();
+	TSharedRef<FDownloadLocalizationTargetFile, ESPMode::ThreadSafe> DownloadTargetFileOp = ILocalizationServiceOperation::Create<FDownloadLocalizationTargetFile>();
+	DownloadTargetFileOp->SetInTargetGuid(AssociatedLocalizationTarget->Settings.Guid);
+	DownloadTargetFileOp->SetInLocale(FPaths::GetBaseFilename(FPaths::GetPath(ArchiveFilePath)));
+	FString Path = FPaths::ProjectSavedDir() / "Temp" / "LastImportFromLocService.po";
+	FPaths::MakePathRelativeTo(Path, *FPaths::ProjectDir());
+	DownloadTargetFileOp->SetInRelativeOutputFilePathAndName(Path);
+
+	GWarn->BeginSlowTask(LOCTEXT("ImportingFromLocalizationService", "Importing Latest from Localization Service..."), true);
+
+	Provider.Execute(DownloadTargetFileOp, TArray<FLocalizationServiceTranslationIdentifier>(), ELocalizationServiceOperationConcurrency::Asynchronous, FLocalizationServiceOperationComplete::CreateSP(this, &FTranslationEditor::DownloadLatestFromLocalizationServiceComplete));
+}
+
+void FTranslationEditor::DownloadLatestFromLocalizationServiceComplete(const FLocalizationServiceOperationRef& Operation, ELocalizationServiceOperationCommandResult::Type Result)
+{
+	check(AssociatedLocalizationTarget.IsValid());
+	TSharedPtr<FDownloadLocalizationTargetFile, ESPMode::ThreadSafe> DownloadLocalizationTargetOp = StaticCastSharedRef<FDownloadLocalizationTargetFile>(Operation);
+	bool bError = !(Result == ELocalizationServiceOperationCommandResult::Succeeded);
+	FText ErrorText = FText::GetEmpty();
+	if (DownloadLocalizationTargetOp.IsValid())
+	{
+		ErrorText = DownloadLocalizationTargetOp->GetOutErrorText();
+	}
+	if (!bError && ErrorText.IsEmpty())
+	{
+		FGuid InTargetGuid;
+		FString InLocale;
+		FString InRelativeOutputFilePathAndName;
+
+		if (DownloadLocalizationTargetOp.IsValid())
+		{
+			InTargetGuid = DownloadLocalizationTargetOp->GetInTargetGuid();
+			InLocale = DownloadLocalizationTargetOp->GetInLocale();
+			InRelativeOutputFilePathAndName = DownloadLocalizationTargetOp->GetInRelativeOutputFilePathAndName();
+
+		}
+		else
+		{
+			bError = true;
+		}
+
+		if (InTargetGuid == AssociatedLocalizationTarget->Settings.Guid && InLocale == FPaths::GetBaseFilename(FPaths::GetPath(ArchiveFilePath)) && !InRelativeOutputFilePathAndName.IsEmpty())
+		{
+			FString AbsoluteFilePathAndName = FPaths::ConvertRelativePathToFull(FPaths::ProjectDir() / InRelativeOutputFilePathAndName);
+			if (FPaths::FileExists(AbsoluteFilePathAndName))
+			{
+				GWarn->StatusUpdate(50, 100, LOCTEXT("DownloadFromLocalizationServiceFinishedNowImporting", "Download from Localization Service Finished, Importing..."));
+				ImportFromPoFile(AbsoluteFilePathAndName);
+			}
+			else
+			{
+				bError = true;
+			}
+		}
+		else
+		{
+			bError = true;
+		}
+
+		if (bError)
+		{
+			if (ErrorText.IsEmpty())
+			{
+				ErrorText = LOCTEXT("DownloadLatestFromLocalizationServiceFileProcessError", "An error occured when processing the file downloaded from the Localization Service.");
+			}
+		}
+	}
+	else
+	{
+		bError = true;
+		if (ErrorText.IsEmpty())
+		{
+			ErrorText = LOCTEXT("DownloadLatestFromLocalizationServiceDownloadError", "An error occured while downloading the file from the Localization Service.");
+		}
+	}
+
+	GWarn->StatusUpdate(100, 100, LOCTEXT("ImportFromLocalizationServiceFinished", "Import from Localization Service Complete!"));
+	GWarn->EndSlowTask();
+
+	if (bError || !ErrorText.IsEmpty())
+	{
+		if (ErrorText.IsEmpty())
+		{
+			ErrorText = LOCTEXT("DownloadLatestFromLocalizationServiceUnspecifiedError", "An unspecified error occured when trying download and import from the Localization Service.");
+		}
+
+		FMessageLog TranslationEditorMessageLog("TranslationEditor");
+		TranslationEditorMessageLog.Error(ErrorText);
+		TranslationEditorMessageLog.Notify(ErrorText);
+	}
+}
+
+void FTranslationEditor::ExportToPortableObjectFormat_Execute()
+{
+	const FString PortableObjectFileDescription = LOCTEXT("PortableObjectFileDescription", "Portable Object File").ToString();
+	const FString PortableObjectFileExtension = TEXT("*.po");
+	const FString FileTypes = FString::Printf(TEXT("%s (%s)|%s"), *PortableObjectFileDescription, *PortableObjectFileExtension, *PortableObjectFileExtension);
+	const FString CultureToEdit = FPaths::GetBaseFilename(FPaths::GetPath(ArchiveFilePath));
+	FString DefaultPath = FPaths::GetPath(LocalizationConfigurationScript::GetDefaultPOPath(AssociatedLocalizationTarget.Get(), CultureToEdit));
+	if (!LastExportFilePath.IsEmpty())
+	{
+		DefaultPath = LastExportFilePath;
+	}
+	TArray<FString> SaveFilenames;
+	IDesktopPlatform* DesktopPlatform = FDesktopPlatformModule::Get();
+	bool bSelected = false;
+	const TSharedPtr<SWindow>& ParentWindow = FSlateApplication::Get().FindWidgetWindow(PreviewTextBlock);
+
+	// Prompt the user for the filename
+	if (DesktopPlatform)
+	{
+		void* ParentWindowWindowHandle = NULL;
+
+		if (ParentWindow.IsValid() && ParentWindow->GetNativeWindow().IsValid())
+		{
+			ParentWindowWindowHandle = ParentWindow->GetNativeWindow()->GetOSWindowHandle();
+		}
+
+		bSelected = DesktopPlatform->SaveFileDialog(
+			ParentWindowWindowHandle,
+			LOCTEXT("ChooseExportLocationWindowTitle", "Choose Export Location").ToString(),
+			DefaultPath,
+			LocalizationConfigurationScript::GetDefaultPOFileName(AssociatedLocalizationTarget.Get()),
+			FileTypes,
+			EFileDialogFlags::None,
+			SaveFilenames
+			);
+	}
+
+	if (bSelected)
+	{
+		LastExportFilePath = FPaths::GetPath(SaveFilenames[0]);
+
+		// Write translation data first to ensure all changes are exported
+		if (DataManager->WriteTranslationData() && ParentWindow.IsValid() && SaveFilenames.Num() > 0)
+		{
+			LocalizationCommandletTasks::ExportTextForCulture(ParentWindow.ToSharedRef(), AssociatedLocalizationTarget.Get(), CultureToEdit, TOptional<FString>(SaveFilenames.Top()));
+		}
+		else
+		{
+			FNotificationInfo Info( LOCTEXT("ExportFailedError", "Translation export failed!") );
+			Info.ExpireDuration = 4.f;
+
+			const TSharedPtr<SNotificationItem> NotificationItem = FSlateNotificationManager::Get().AddNotification(Info);
+			if (NotificationItem.IsValid())
+			{
+				NotificationItem->SetCompletionState(SNotificationItem::CS_Fail);
+			}
+		}
+	}
+}
+
+void FTranslationEditor::ImportFromPortableObjectFormat_Execute()
+{
+	const FString PortableObjectFileDescription = LOCTEXT("PortableObjectFileDescription", "Portable Object File").ToString();
+	const FString PortableObjectFileExtension = TEXT("*.po");
+	const FString FileTypes = FString::Printf(TEXT("%s (%s)|%s"), *PortableObjectFileDescription, *PortableObjectFileExtension, *PortableObjectFileExtension);
+	const FString CultureToEdit = FPaths::GetBaseFilename(FPaths::GetPath(ArchiveFilePath));
+	FString DefaultPath = FPaths::GetPath(LocalizationConfigurationScript::GetDefaultPOPath(AssociatedLocalizationTarget.Get(), CultureToEdit));
+	if (!LastImportFilePath.IsEmpty())
+	{
+		DefaultPath = LastImportFilePath;
+	}
+	TArray<FString> OpenFilenames;
+	IDesktopPlatform* DesktopPlatform = FDesktopPlatformModule::Get();
+
+	bool bOpened = false;
+	if (DesktopPlatform)
+	{
+		void* ParentWindowWindowHandle = NULL;
+
+		const TSharedPtr<SWindow>& ParentWindow = FSlateApplication::Get().FindWidgetWindow(PreviewTextBlock);
+		if (ParentWindow.IsValid() && ParentWindow->GetNativeWindow().IsValid())
+		{
+			ParentWindowWindowHandle = ParentWindow->GetNativeWindow()->GetOSWindowHandle();
+		}
+
+		bOpened = DesktopPlatform->OpenFileDialog(
+			ParentWindowWindowHandle,
+			LOCTEXT("ChooseImportLocationWindowTitle", "Choose File to Import").ToString(),
+			DefaultPath,
+			LocalizationConfigurationScript::GetDefaultPOFileName(AssociatedLocalizationTarget.Get()),
+			FileTypes,
+			EFileDialogFlags::None,
+			OpenFilenames
+			);
+	}
+
+	if (bOpened && OpenFilenames.Num() > 0)
+	{
+		FString FileToImport = OpenFilenames[0];
+
+		ImportFromPoFile(FileToImport);
+	}
+}
+
+void FTranslationEditor::ImportFromPoFile(FString FileToImport)
+{
+	LastImportFilePath = FPaths::GetPath(FileToImport);
+
+	// Write translation data first to ensure all changes are exported
+	const FString CultureToEdit = FPaths::GetBaseFilename(FPaths::GetPath(ArchiveFilePath));
+	const TSharedPtr<SWindow>& ParentWindow = FSlateApplication::Get().FindWidgetWindow(PreviewTextBlock);
+
+	if (DataManager->WriteTranslationData(true) && ParentWindow.IsValid())
+	{
+		if (LocalizationCommandletTasks::ImportTextForCulture(ParentWindow.ToSharedRef(), AssociatedLocalizationTarget.Get(), CultureToEdit, TOptional<FString>(FileToImport)))
+		{
+			DataManager->LoadFromArchive(DataManager->GetAllTranslationsArray(), true, true);
+
+			TabManager->TryInvokeTab(ChangedOnImportTabId);
+			ChangedOnImportPropertyTable->SetObjects((TArray<UObject*>&)DataManager->GetChangedOnImportArray());
+			// Need to re-add the columns we want to display
+			ChangedOnImportPropertyTable->AddColumn((TWeakFieldPtr<FProperty>)FindFProperty<FProperty>(UTranslationUnit::StaticClass(), "Source"));
+			ChangedOnImportPropertyTable->AddColumn((TWeakFieldPtr<FProperty>)FindFProperty<FProperty>(UTranslationUnit::StaticClass(), "TranslationBeforeImport"));
+			ChangedOnImportPropertyTable->AddColumn((TWeakFieldPtr<FProperty>)FindFProperty<FProperty>(UTranslationUnit::StaticClass(), "Translation"));
+		}
+	}
+	else
+	{
+		FNotificationInfo Info( LOCTEXT("ImportFailedError", "Translation import failed!") );
+		Info.ExpireDuration = 4.f;
+
+		const TSharedPtr<SNotificationItem> NotificationItem = FSlateNotificationManager::Get().AddNotification(Info);
+		if (NotificationItem.IsValid())
+		{
+			NotificationItem->SetCompletionState(SNotificationItem::CS_Fail);
+		}
+	}
+}
+
+void FTranslationEditor::OnFilterTextChanged(const FText& InFilterText)
+{
+
+}
+
+void FTranslationEditor::OnFilterTextCommitted(const FText& InFilterText, ETextCommit::Type CommitInfo)
+{
+	const FString InFilterString = InFilterText.ToString();
+
+	if (CommitInfo == ETextCommit::OnEnter)
+	{
+		if (InFilterString != CurrentSearchFilter)
+		{
+			CurrentSearchFilter = InFilterString;
+
+			DataManager->PopulateSearchResultsUsingFilter(InFilterString);
+
+			if (SearchPropertyTable.IsValid())
+			{
+				SearchPropertyTable->SetObjects((TArray<UObject*>&)DataManager->GetSearchResultsArray());
+
+				// Need to re-add the columns we want to display
+				SearchPropertyTable->AddColumn((TWeakFieldPtr<FProperty>)FindFProperty<FProperty>(UTranslationUnit::StaticClass(), "Source"));
+				SearchPropertyTable->AddColumn((TWeakFieldPtr<FProperty>)FindFProperty<FProperty>(UTranslationUnit::StaticClass(), "Translation"));
+
+				TArray<TSharedRef<IPropertyTableColumn>> Columns = SearchPropertyTable->GetColumns();
+				for (TSharedRef<IPropertyTableColumn> Column : Columns)
+				{
+					Column->SetFrozen(true);
+				}
+			}
+		}
+	}
+}
+
+void FTranslationEditor::OpenSearchTab_Execute()
+{
+	TabManager->TryInvokeTab(SearchTabId);
+}
+
+FReply FTranslationEditor::OnGetHistoryButtonClicked()
+{
+	// Load the actual history data
+	DataManager->GetHistoryForTranslationUnits();
+
+	// Items might have moved from Untranslated to review, so refresh the view of both tables
+	if (UntranslatedPropertyTable.IsValid())
+	{
+		UntranslatedPropertyTable->SetObjects((TArray<UObject*>&)DataManager->GetUntranslatedArray());
+
+		// Need to re-add the columns we want to display
+		UntranslatedPropertyTable->AddColumn((TWeakFieldPtr<FProperty>)FindFProperty<FProperty>(UTranslationUnit::StaticClass(), "Source"));
+		UntranslatedPropertyTable->AddColumn((TWeakFieldPtr<FProperty>)FindFProperty<FProperty>(UTranslationUnit::StaticClass(), "Translation"));
+
+		TArray<TSharedRef<IPropertyTableColumn>> Columns = UntranslatedPropertyTable->GetColumns();
+		for (TSharedRef<IPropertyTableColumn> Column : Columns)
+		{
+			Column->SetFrozen(true);
+		}
+	}
+
+	if (ReviewPropertyTable.IsValid())
+	{
+		ReviewPropertyTable->SetObjects((TArray<UObject*>&)DataManager->GetReviewArray());
+
+		// Need to re-add the columns we want to display
+		ReviewPropertyTable->AddColumn((TWeakFieldPtr<FProperty>)FindFProperty<FProperty>(UTranslationUnit::StaticClass(), "Source"));
+		ReviewPropertyTable->AddColumn((TWeakFieldPtr<FProperty>)FindFProperty<FProperty>(UTranslationUnit::StaticClass(), "Translation"));
+		ReviewPropertyTable->AddColumn((TWeakFieldPtr<FProperty>)FindFProperty<FProperty>(UTranslationUnit::StaticClass(), "HasBeenReviewed"));
+
+		TArray<TSharedRef<IPropertyTableColumn>> Columns = ReviewPropertyTable->GetColumns();
+		for (TSharedRef<IPropertyTableColumn> Column : Columns)
+		{
+			FString ColumnId = Column->GetId().ToString();
+			if (ColumnId == "HasBeenReviewed")
+			{
+				Column->SetWidth(120);
+				Column->SetSizeMode(EPropertyTableColumnSizeMode::Fixed);
+			}
+			// Freeze columns, don't want user to remove them
+			Column->SetFrozen(true);
+		}
+	}
+
+	// Make sure all UI is refreshed
+	RefreshUI();
+
+	// Make sure current selection is reflected
+	UpdateUntranslatedSelection();
+	UpdateNeedsReviewSelection();
+	UpdateCompletedSelection();
+	UpdateSearchSelection();
+
+	return FReply::Handled();
+}
+
+#undef LOCTEXT_NAMESPACE
+
+#endif // WITH_EDITOR
diff --git a/Engine/Source/Developer/TranslationEditor/Private/TranslationEditor.h b/Engine/Source/Developer/TranslationEditor/Private/TranslationEditor.h
new file mode 100644
--- /dev/null
+++ b/Engine/Source/Developer/TranslationEditor/Private/TranslationEditor.h
@@ -0,0 +1,312 @@
+// Copyright Epic Games, Inc. All Rights Reserved.
+#pragma once
+#if WITH_EDITOR
+
+#include "CoreMinimal.h"
+#include "Input/Reply.h"
+#include "Fonts/SlateFontInfo.h"
+#include "Widgets/DeclarativeSyntaxSupport.h"
+#include "Widgets/Text/STextBlock.h"
+#include "Toolkits/IToolkitHost.h"
+#include "TranslationDataManager.h"
+#include "Styling/AppStyle.h"
+#include "ITranslationEditor.h"
+#include "IPropertyTable.h"
+#include "CustomFontColumn.h"
+#include "TranslationUnit.h"
+#include "LocalizationConfigurationScript.h"
+
+class IPropertyTableRow;
+class IPropertyTableWidgetHandle;
+
+class TRANSLATIONEDITOR_API FTranslationEditor :  public ITranslationEditor
+{
+public:
+	/**
+	 *	Creates a new FTranslationEditor and calls Initialize
+	 */
+	static TSharedRef< FTranslationEditor > Create(TSharedRef< FTranslationDataManager > DataManager, const FString& InManifestFile, const FString& InArchiveFile)
+	{
+		TSharedRef< FTranslationEditor > TranslationEditor = MakeShareable(new FTranslationEditor(DataManager, InManifestFile, InArchiveFile, nullptr));
+
+		// Some stuff that needs to use the "this" pointer is done in Initialize (because it can't be done in the constructor)
+		TranslationEditor->Initialize();
+
+		for (UTranslationUnit* TranslationUnit : DataManager->GetAllTranslationsArray())
+		{
+			// Set up a property changed event to trigger a write of the translation data when TranslationUnit property changes
+			TranslationUnit->OnPropertyChanged().AddSP(DataManager, &FTranslationDataManager::HandlePropertyChanged);
+		}
+
+		return TranslationEditor;
+	}
+	
+	static TSharedRef< FTranslationEditor > Create(TSharedRef< FTranslationDataManager > DataManager, ULocalizationTarget* const LocalizationTarget, const FString& CultureToEdit)
+	{
+		check(LocalizationTarget);
+		
+		TSharedRef< FTranslationEditor > TranslationEditor = MakeShareable(new FTranslationEditor(DataManager, LocalizationConfigurationScript::GetManifestPath(LocalizationTarget), LocalizationConfigurationScript::GetArchivePath(LocalizationTarget, CultureToEdit), LocalizationTarget));
+
+		// Some stuff that needs to use the "this" pointer is done in Initialize (because it can't be done in the constructor)
+		TranslationEditor->Initialize();
+
+		for (UTranslationUnit* TranslationUnit : DataManager->GetAllTranslationsArray())
+		{
+			// Set up a property changed event to trigger a write of the translation data when TranslationUnit property changes
+			TranslationUnit->OnPropertyChanged().AddSP(DataManager, &FTranslationDataManager::HandlePropertyChanged);
+		}
+
+		return TranslationEditor;
+	}
+
+	virtual ~FTranslationEditor() {}
+
+	virtual void RegisterTabSpawners(const TSharedRef<class FTabManager>& TabManager) override;
+	virtual void UnregisterTabSpawners(const TSharedRef<class FTabManager>& TabManager) override;
+
+	/**
+	 * Edits the specified table
+	 *
+	 * @param	Mode					Asset editing mode for this editor (standalone or world-centric)
+	 * @param	InitToolkitHost			When Mode is WorldCentric, this is the level editor instance to spawn this editor within
+	 * @param	ObjectToEdit			The object to edit
+	 */
+	void InitTranslationEditor( const EToolkitMode::Type Mode, const TSharedPtr< class IToolkitHost >& InitToolkitHost );
+
+	/** IToolkit interface */
+	virtual FName GetToolkitFName() const override;
+	virtual FText GetBaseToolkitName() const override;
+	virtual FText GetToolkitName() const override;
+	virtual FText GetToolkitToolTipText() const override;
+	virtual FString GetWorldCentricTabPrefix() const override;
+	virtual FLinearColor GetWorldCentricTabColorScale() const override;
+
+protected:
+	/** Called when "Save" is clicked for this asset */
+	virtual void SaveAsset_Execute() override;
+
+private:
+
+	FTranslationEditor(TSharedRef< FTranslationDataManager > InDataManager, const FString& InManifestFile, const FString& InArchiveFile, ULocalizationTarget* const LocalizationTarget)
+		: ITranslationEditor(InManifestFile, InArchiveFile, LocalizationTarget)
+		, DataManager(InDataManager)
+		, SourceFont(FAppStyle::GetFontStyle( PropertyTableConstants::NormalFontStyle ))
+		, TranslationTargetFont(FAppStyle::GetFontStyle( PropertyTableConstants::NormalFontStyle ))
+		, SourceColumn(MakeShareable(new FCustomFontColumn(SourceFont)))
+		, TranslationColumn(MakeShareable(new FCustomFontColumn(TranslationTargetFont)))
+		, PreviewTextBlock(SNew(STextBlock)
+				.Text(FText::FromString(""))
+				.Font(TranslationTargetFont))
+		, NamespaceTextBlock(SNew(STextBlock)
+				.Text(FText::FromString("")))
+	{}
+
+	/** Does some things we can't do in the constructor because we can't get a SharedRef to "this" there */ 
+	void Initialize();
+
+	/**	Spawns the untranslated tab */
+	TSharedRef<SDockTab> SpawnTab_Untranslated( const FSpawnTabArgs& Args );
+	/** Getter function bound to the label attribute for the untranslated tab. */
+	FText GetTabLabel_Untranslated() const;
+	/**	Spawns the review tab */
+	TSharedRef<SDockTab> SpawnTab_Review( const FSpawnTabArgs& Args );
+	/** Getter function bound to the label attribute for the review tab. */
+	FText GetTabLabel_Review() const;
+	/**	Spawns the completed tab */
+	TSharedRef<SDockTab> SpawnTab_Completed( const FSpawnTabArgs& Args );
+	/** Getter function bound to the label attribute for the completed tab. */
+	FText GetTabLabel_Completed() const;
+	/**	Spawns the preview tab */
+	TSharedRef<SDockTab> SpawnTab_Preview( const FSpawnTabArgs& Args );
+	/**	Spawns the context tab */
+	TSharedRef<SDockTab> SpawnTab_Context( const FSpawnTabArgs& Args );
+	/**	Spawns the history tab */
+	TSharedRef<SDockTab> SpawnTab_History( const FSpawnTabArgs& Args );
+	/**	Spawns the search tab */
+	TSharedRef<SDockTab> SpawnTab_Search( const FSpawnTabArgs& Args );
+	/**	Spawns the Changed on Import tab */
+	TSharedRef<SDockTab> SpawnTab_ChangedOnImport(const FSpawnTabArgs& Args);
+
+	/** Map actions for the UI_COMMANDS */
+	void MapActions();
+
+	/** Change the font for the source language */
+	void ChangeSourceFont();
+
+	/** For button delegate */
+	FReply ChangeSourceFont_FReply()
+	{
+		ChangeSourceFont();
+
+		return FReply::Handled();
+	}
+
+	/** Change the font for the target translation language */
+	void ChangeTranslationTargetFont();
+
+	/** For button delegate */
+	FReply ChangeTranslationTargetFont_FReply()
+	{
+		ChangeTranslationTargetFont();
+
+		return FReply::Handled();
+	}
+
+	/** Called on SpinBox OnValueCommitted */
+	void OnSourceFontSizeCommitt(int32 NewFontSize, ETextCommit::Type)
+	{
+		SourceFont.Size = NewFontSize;
+		RefreshUI();
+	}
+
+	void OnTranslationTargetFontSizeCommitt(int32 NewFontSize, ETextCommit::Type)
+	{
+		TranslationTargetFont.Size = NewFontSize;
+		RefreshUI();
+	}
+
+	/** Open the file dialog prompt (at the DefaultFile location) to allow the user to pick a font, then return the user's selection, and a boolean of whether something was selected */
+	bool OpenFontPicker(const FString DefaultFile, FString& OutFile);
+
+	/** Reset all of the UI after a new font is chosen */
+	void RefreshUI();
+
+	/** Update content when a new translation unit selection is made */
+	void UpdateTranslationUnitSelection(TSet<TSharedRef<IPropertyTableRow>>& SelectedRows);
+
+	/** Update content when a new translation unit selection is made in the Untranslated PropertyTable */
+	void UpdateUntranslatedSelection();
+
+	/** Update content when a new translation unit selection is made in the Needs Review PropertyTable */
+	void UpdateNeedsReviewSelection();
+
+	/** Update content when a new translation unit selection is made in the Completed PropertyTable */
+	void UpdateCompletedSelection();
+
+	/** Update content when a new translation unit selection is made in the Search PropertyTable */
+	void UpdateSearchSelection();
+
+	/** Update content when a new translation unit selection is made in the ChangedOnImport PropertyTable */
+	void UpdateChangedOnImportSelection();
+
+	/** Update content when a new context selection is made */
+	void UpdateContextSelection();
+
+	/** Called when "Preview in Editor" is clicked for this Localization Target */
+	void PreviewAllTranslationsInEditor_Execute();
+	
+	/** Called when "Import from Localization Service" is clicked for this Localization Target */
+	void ImportLatestFromLocalizationService_Execute();
+
+	/** Callback for when the Localization Service operation started when "Import from Localization Service" was clicked finishes */
+	void DownloadLatestFromLocalizationServiceComplete(const FLocalizationServiceOperationRef& Operation, ELocalizationServiceOperationCommandResult::Type Result);
+
+	/** Called when "Export to .PO" is clicked for this Localization Target */
+	void ExportToPortableObjectFormat_Execute();
+
+	/** Called when "Import from .PO" is clicked for this Localization Target */
+	void ImportFromPortableObjectFormat_Execute();
+
+	/** Import from the specified .po file into this localization target */
+	void ImportFromPoFile(FString FileToImport);
+
+	/** Open the search tab */
+	void OpenSearchTab_Execute();
+
+	/** Called when the filter text in the SearchBox is changed */
+	void OnFilterTextChanged(const FText& InFilterText);
+
+	/** Called when text is committed to the SearchBox */
+	void OnFilterTextCommitted(const FText& InFilterText, ETextCommit::Type CommitInfo);
+
+	/** Called when the get history button is clicked to retrieve history from source control */
+	FReply OnGetHistoryButtonClicked();
+
+	/**	The tab id for the untranslated tab */
+	static const FName UntranslatedTabId;
+	/**	The tab id for the review tab */
+	static const FName ReviewTabId;
+	/**	The tab id for the completed tab */
+	static const FName CompletedTabId;
+	/**	The tab id for the preview tab */
+	static const FName PreviewTabId;
+	/**	The tab id for the context tab */
+	static const FName ContextTabId;
+	/**	The tab id for the history tab */
+	static const FName HistoryTabId;
+	/**	The tab id for the search tab */
+	static const FName SearchTabId;
+	/**	The tab id for the changed on import tab */
+	static const FName ChangedOnImportTabId;
+
+	/** The Untranslated Tab */
+	TWeakPtr<SDockTab> UntranslatedTab;
+	/** The Review Tab */
+	TWeakPtr<SDockTab> ReviewTab;
+	/** The Review Tab */
+	TWeakPtr<SDockTab> CompletedTab;
+	/** The Search Tab */
+	TWeakPtr<SDockTab> SearchTab;
+	/** The Changed On Import Tab */
+	TWeakPtr<SDockTab> ChangedOnImportTab;
+
+	/** Search box for searching the source and translation strings */
+	TSharedPtr<SSearchBox> SearchBox;
+	/** Current search filter */
+	FString CurrentSearchFilter;
+
+	/** Manages the reading and writing of data to file */
+	TSharedRef< FTranslationDataManager > DataManager;
+
+	/** The table of untranslated items */
+	TSharedPtr< class IPropertyTable > UntranslatedPropertyTable;
+	/** The table of translations to review */
+	TSharedPtr< class IPropertyTable > ReviewPropertyTable;
+	/** The table of completed translations */
+	TSharedPtr< class IPropertyTable > CompletedPropertyTable;
+	/** The table of context information */
+	TSharedPtr< class IPropertyTable > ContextPropertyTable;
+	/** The table of previous revision information */
+	TSharedPtr< class IPropertyTable > HistoryPropertyTable;
+	/** The table of search results */
+	TSharedPtr< class IPropertyTable > SearchPropertyTable;
+	/** The table of changed on import results */
+	TSharedPtr< class IPropertyTable > ChangedOnImportPropertyTable;
+
+	/** The slate widget table of untranslated items */
+	TSharedPtr< class IPropertyTableWidgetHandle > UntranslatedPropertyTableWidgetHandle;
+	/** The slate widget table of translations to review */
+	TSharedPtr< class IPropertyTableWidgetHandle > ReviewPropertyTableWidgetHandle;
+	/** The slate widget table of completed items */
+	TSharedPtr< class IPropertyTableWidgetHandle > CompletedPropertyTableWidgetHandle;
+	/** The slate widget table of contexts for this item */
+	TSharedPtr< class IPropertyTableWidgetHandle > ContextPropertyTableWidgetHandle;
+	/** The slate widget table of previous revision information */
+	TSharedPtr< class IPropertyTableWidgetHandle > HistoryPropertyTableWidgetHandle;
+	/** The slate widget table of search results */
+	TSharedPtr< class IPropertyTableWidgetHandle > SearchPropertyTableWidgetHandle;
+	/** The slate widget table of translations that changed on import */
+	TSharedPtr< class IPropertyTableWidgetHandle > ChangedOnImportPropertyTableWidgetHandle;
+
+	/** Font to use for the source language */
+	FSlateFontInfo SourceFont;
+	/** Font to use for the translation target language */
+	FSlateFontInfo TranslationTargetFont;
+
+	/** Custom FontColumn for columns that display source text */
+	TSharedRef<FCustomFontColumn> SourceColumn;
+	/** Custom FontColumn for columns that display source text */
+	TSharedRef<FCustomFontColumn> TranslationColumn;
+
+	/** Text block for previewing the currently selected translation */
+	TSharedRef<STextBlock> PreviewTextBlock;
+	/** Text block displaying the namespace of the currently selected translation unit */
+	TSharedRef<STextBlock> NamespaceTextBlock;
+
+	/** Used to remember the location of the file the user last exported to */
+	FString LastExportFilePath;
+	/** Used to remember the location of the file the user last imported */
+	FString LastImportFilePath;
+};
+
+#endif // WITH_EDITOR
diff --git a/Engine/Source/Developer/TranslationEditor/Private/TranslationEditorMenu.cpp b/Engine/Source/Developer/TranslationEditor/Private/TranslationEditorMenu.cpp
new file mode 100644
--- /dev/null
+++ b/Engine/Source/Developer/TranslationEditor/Private/TranslationEditorMenu.cpp
@@ -0,0 +1,120 @@
+// Copyright Epic Games, Inc. All Rights Reserved.
+#if WITH_EDITOR
+
+#include "TranslationEditorMenu.h"
+#include "Widgets/DeclarativeSyntaxSupport.h"
+#include "Settings/EditorExperimentalSettings.h"
+#include "TranslationPickerWidget.h"
+#include "Framework/MultiBox/MultiBoxBuilder.h"
+#include "TranslationEditor.h"
+
+#define LOCTEXT_NAMESPACE "TranslationEditorToolbar"
+
+void FTranslationEditorMenu::FillTranslationMenu( FMenuBuilder& MenuBuilder/*, FTranslationEditor& TranslationEditor*/ )
+{
+	MenuBuilder.BeginSection("Font", LOCTEXT("Translation_FontHeading", "Font"));
+	{
+		MenuBuilder.AddMenuEntry( FTranslationEditorCommands::Get().ChangeSourceFont );
+		MenuBuilder.AddMenuEntry( FTranslationEditorCommands::Get().ChangeTranslationTargetFont );
+		MenuBuilder.AddMenuEntry( FTranslationEditorCommands::Get().PreviewAllTranslationsInEditor );
+		MenuBuilder.AddMenuEntry( FTranslationEditorCommands::Get().ImportLatestFromLocalizationService );
+		MenuBuilder.AddMenuEntry( FTranslationEditorCommands::Get().ExportToPortableObjectFormat );
+		MenuBuilder.AddMenuEntry( FTranslationEditorCommands::Get().ImportFromPortableObjectFormat );
+		MenuBuilder.AddMenuEntry( FTranslationEditorCommands::Get().OpenSearchTab );
+		if (GetDefault<UEditorExperimentalSettings>()->bEnableTranslationPicker)
+		{
+			MenuBuilder.AddMenuEntry(FTranslationEditorCommands::Get().OpenTranslationPicker);
+		}
+	}
+	MenuBuilder.EndSection();
+}
+
+void FTranslationEditorMenu::SetupTranslationEditorMenu( TSharedPtr< FExtender > Extender, FTranslationEditor& TranslationEditor)
+{
+	// Add additional editor menu
+	{
+		struct Local
+		{
+			static void AddSaveMenuOption( FMenuBuilder& MenuBuilder )
+			{
+				MenuBuilder.AddMenuEntry( FTranslationEditorCommands::Get().SaveTranslations, "SaveTranslations", TAttribute<FText>(), TAttribute<FText>(), FSlateIcon(FAppStyle::GetAppStyleSetName(), "AssetEditor.SaveAsset") );
+			}
+
+			static void AddTranslationEditorMenu( FMenuBarBuilder& MenuBarBuilder )
+			{
+				// View
+				MenuBarBuilder.AddPullDownMenu( 
+					LOCTEXT("TranslationMenu", "Translation"),
+					LOCTEXT("TranslationMenu_ToolTip", "Open the Translation menu"),
+					FNewMenuDelegate::CreateStatic( &FTranslationEditorMenu::FillTranslationMenu ),
+					"View");
+			}
+		};
+
+		Extender->AddMenuExtension(
+			"FileLoadAndSave",
+			EExtensionHook::First,
+			TranslationEditor.GetToolkitCommands(),
+			FMenuExtensionDelegate::CreateStatic( &Local::AddSaveMenuOption ) );
+
+		Extender->AddMenuBarExtension(
+			"Edit",
+			EExtensionHook::After,
+			TranslationEditor.GetToolkitCommands(), 
+			FMenuBarExtensionDelegate::CreateStatic( &Local::AddTranslationEditorMenu ) );
+	}
+}
+
+void FTranslationEditorMenu::SetupTranslationEditorToolbar( TSharedPtr< FExtender > Extender, FTranslationEditor& TranslationEditor )
+{
+	struct Local
+	{
+		static void AddToolbarButtons( FToolBarBuilder& ToolbarBuilder )
+		{
+			ToolbarBuilder.AddToolBarButton(
+				FTranslationEditorCommands::Get().SaveTranslations, "SaveTranslations", TAttribute<FText>(), TAttribute<FText>(), FSlateIcon(FAppStyle::GetAppStyleSetName(), "AssetEditor.SaveAsset"));
+			ToolbarBuilder.AddToolBarButton(
+				FTranslationEditorCommands::Get().PreviewAllTranslationsInEditor, "PreviewTranslationsInEditor", TAttribute<FText>(), TAttribute<FText>(), FSlateIcon(FAppStyle::GetAppStyleSetName(), "TranslationEditor.PreviewInEditor")); 
+			ToolbarBuilder.AddToolBarButton(
+				FTranslationEditorCommands::Get().ImportLatestFromLocalizationService, "ImportLatestFromLocalizationService", TAttribute<FText>(), TAttribute<FText>(), FSlateIcon(FAppStyle::GetAppStyleSetName(), "TranslationEditor.ImportLatestFromLocalizationService"));
+			ToolbarBuilder.AddToolBarButton(
+				FTranslationEditorCommands::Get().ExportToPortableObjectFormat, "ExportToPortableObjectFormat", TAttribute<FText>(), TAttribute<FText>(), FSlateIcon(FAppStyle::GetAppStyleSetName(), "TranslationEditor.Export"));
+			ToolbarBuilder.AddToolBarButton(
+				FTranslationEditorCommands::Get().ImportFromPortableObjectFormat, "ImportFromPortableObjectFormat", TAttribute<FText>(), TAttribute<FText>(), FSlateIcon(FAppStyle::GetAppStyleSetName(), "TranslationEditor.Import"));
+			ToolbarBuilder.AddToolBarButton(
+				FTranslationEditorCommands::Get().OpenSearchTab, "OpenSearchTab", TAttribute<FText>(), TAttribute<FText>(), FSlateIcon(FAppStyle::GetAppStyleSetName(), "TranslationEditor.Search"));
+			if (GetDefault<UEditorExperimentalSettings>()->bEnableTranslationPicker)
+			{
+				ToolbarBuilder.AddWidget(SNew(STranslationWidgetPicker));
+			}
+		}
+	};
+
+	Extender->AddToolBarExtension(
+		"Asset",
+		EExtensionHook::First,
+		TranslationEditor.GetToolkitCommands(),
+		FToolBarExtensionDelegate::CreateStatic( &Local::AddToolbarButtons ) );
+}
+
+
+//////////////////////////////////////////////////////////////////////////
+// FTranslationEditorCommands
+
+void FTranslationEditorCommands::RegisterCommands() 
+{
+	UI_COMMAND( ChangeSourceFont, "Change Source Font", "Change the Font for the Source Lanugage", EUserInterfaceActionType::Button, FInputChord() );
+	UI_COMMAND( ChangeTranslationTargetFont, "Change Translation Font", "Change the Translation Target Language Font", EUserInterfaceActionType::Button, FInputChord() );
+	UI_COMMAND( SaveTranslations, "Save", "Saves the translations to file", EUserInterfaceActionType::Button, FInputChord() );
+	UI_COMMAND( PreviewAllTranslationsInEditor, "Preview in Editor", "Preview All Translations in the Editor UI", EUserInterfaceActionType::Button, FInputChord());
+	UI_COMMAND( ImportLatestFromLocalizationService, "Import from Translation Service", "Download and Import Latest Translations From Localization Service. (Localization Service settings can be modified in the Localization Dashboard)", EUserInterfaceActionType::Button, FInputChord());
+	UI_COMMAND( ExportToPortableObjectFormat, "Export to .PO", "Export to Portable Object Format", EUserInterfaceActionType::Button, FInputChord() );
+	UI_COMMAND( ImportFromPortableObjectFormat, "Import from .PO", "Import from Portable Object Format", EUserInterfaceActionType::Button, FInputChord() );
+	UI_COMMAND( OpenSearchTab, "Search", "Search Source and Translation Strings", EUserInterfaceActionType::Button, FInputChord() );
+	UI_COMMAND( OpenTranslationPicker, "Translation Picker", "Open the Translation Picker to Modify Editor Translations", EUserInterfaceActionType::Button, FInputChord() )
+}
+
+
+#undef LOCTEXT_NAMESPACE
+
+#endif // WITH_EDITOR
diff --git a/Engine/Source/Developer/TranslationEditor/Private/TranslationEditorModule.cpp b/Engine/Source/Developer/TranslationEditor/Private/TranslationEditorModule.cpp
new file mode 100644
--- /dev/null
+++ b/Engine/Source/Developer/TranslationEditor/Private/TranslationEditorModule.cpp
@@ -0,0 +1,99 @@
+// Copyright Epic Games, Inc. All Rights Reserved.
+
+#include "TranslationEditorModule.h"
+#include "Misc/FeedbackContext.h"
+#include "Modules/ModuleManager.h"
+#include "TranslationDataManager.h"
+#include "TranslationPickerWidget.h"
+#include "TranslationEditor.h"
+#include "MessageLogModule.h"
+#include "ITranslationEditor.h"
+
+class FTranslationEditor;
+
+IMPLEMENT_MODULE( FTranslationEditorModule, TranslationEditor );
+
+#define LOCTEXT_NAMESPACE "TranslationEditorModule"
+
+#if WITH_EDITOR
+const FName FTranslationEditorModule::TranslationEditorAppIdentifier( TEXT( "TranslationEditorApp" ) );
+#endif // WITH_EDITOR
+
+void FTranslationEditorModule::StartupModule()
+{
+#if WITH_UNREAL_DEVELOPER_TOOLS
+	// create a message log for source control to use
+	FMessageLogModule& MessageLogModule = FModuleManager::LoadModuleChecked<FMessageLogModule>("MessageLog");
+	MessageLogModule.RegisterLogListing("TranslationEditor", LOCTEXT("TranslationEditorLogLabel", "Translation Editor"));
+#endif
+
+#if WITH_EDITOR
+	MenuExtensibilityManager = MakeShareable(new FExtensibilityManager);
+	ToolbarExtensibilityManager = MakeShareable(new FExtensibilityManager);
+#endif // WITH_EDITOR
+}
+
+void FTranslationEditorModule::ShutdownModule()
+{
+#if WITH_EDITOR
+	MenuExtensibilityManager.Reset();
+#endif // WITH_EDITOR
+	TranslationPickerManager::ClosePickerWindow();
+
+#if WITH_UNREAL_DEVELOPER_TOOLS
+	// unregister message log
+	FMessageLogModule* MessageLogModule = FModuleManager::LoadModulePtr<FMessageLogModule>("MessageLog");
+	if (MessageLogModule)
+	{
+		MessageLogModule->UnregisterLogListing("TranslationEditor");
+	}
+#endif
+}
+
+#if WITH_EDITOR
+TSharedRef<FTranslationEditor> FTranslationEditorModule::CreateTranslationEditor(const FString& ManifestFile, const FString& NativeArchiveFile, const FString& ArchiveFileToEdit, bool& OutLoadedSuccessfully)
+{
+	TSharedRef< FTranslationDataManager > DataManager = MakeShareable( new FTranslationDataManager(ManifestFile, NativeArchiveFile, ArchiveFileToEdit) );
+	OutLoadedSuccessfully = DataManager->GetLoadedSuccessfully();
+
+	GWarn->BeginSlowTask(LOCTEXT("BuildingUserInterface", "Building Translation Editor UI..."), true);
+
+	TSharedRef< FTranslationEditor > NewTranslationEditor(FTranslationEditor::Create(DataManager, ManifestFile, ArchiveFileToEdit));
+	NewTranslationEditor->InitTranslationEditor(EToolkitMode::Standalone, TSharedPtr<IToolkitHost>());
+
+	GWarn->EndSlowTask();
+
+	return NewTranslationEditor;
+}
+
+TSharedRef<FTranslationEditor> FTranslationEditorModule::CreateTranslationEditor(ULocalizationTarget* const LocalizationTarget, const FString& CultureToEdit, bool& OutLoadedSuccessfully)
+{
+	const FString ManifestFile = LocalizationConfigurationScript::GetManifestPath(LocalizationTarget);
+	FString NativeCultureName;
+	if (LocalizationTarget->Settings.SupportedCulturesStatistics.IsValidIndex(LocalizationTarget->Settings.NativeCultureIndex))
+	{
+		NativeCultureName = LocalizationTarget->Settings.SupportedCulturesStatistics[LocalizationTarget->Settings.NativeCultureIndex].CultureName;
+	}
+	const FString NativeArchiveFile = NativeCultureName.IsEmpty() ? FString() : LocalizationConfigurationScript::GetArchivePath(LocalizationTarget, NativeCultureName);
+	const FString ArchiveFileToEdit = LocalizationConfigurationScript::GetArchivePath(LocalizationTarget, CultureToEdit);
+
+	TSharedRef< FTranslationDataManager > DataManager = MakeShareable( new FTranslationDataManager(LocalizationTarget, CultureToEdit) );
+	OutLoadedSuccessfully = DataManager->GetLoadedSuccessfully();
+
+	GWarn->BeginSlowTask(LOCTEXT("BuildingUserInterface", "Building Translation Editor UI..."), true);
+
+	TSharedRef< FTranslationEditor > NewTranslationEditor(FTranslationEditor::Create(DataManager, LocalizationTarget, CultureToEdit));
+	NewTranslationEditor->InitTranslationEditor(EToolkitMode::Standalone, TSharedPtr<IToolkitHost>());
+
+	GWarn->EndSlowTask();
+
+	return NewTranslationEditor;
+}
+#endif // WITH_EDITOR
+
+void FTranslationEditorModule::OpenTranslationPicker()
+{
+	ITranslationEditor::OpenTranslationPicker();
+}
+
+#undef LOCTEXT_NAMESPACE
diff --git a/Engine/Source/Developer/TranslationEditor/Private/TranslationPickerEditWindow.cpp b/Engine/Source/Developer/TranslationEditor/Private/TranslationPickerEditWindow.cpp
new file mode 100644
--- /dev/null
+++ b/Engine/Source/Developer/TranslationEditor/Private/TranslationPickerEditWindow.cpp
@@ -0,0 +1,614 @@
+// Copyright Epic Games, Inc. All Rights Reserved.
+
+#include "TranslationPickerEditWindow.h"
+
+#include "Framework/Application/IInputProcessor.h"
+#include "Framework/Application/SlateApplication.h"
+#include "HAL/PlatformCrt.h"
+#include "HAL/PlatformApplicationMisc.h"
+#include "ILocalizationServiceModule.h"
+#include "ILocalizationServiceProvider.h"
+#include "Input/Events.h"
+#include "InputCoreTypes.h"
+#include "Internationalization/Internationalization.h"
+#include "Internationalization/TextKey.h"
+#include "Internationalization/TextLocalizationManager.h"
+#include "Internationalization/TextNamespaceUtil.h"
+#include "Layout/BasicLayoutWidgetSlot.h"
+#include "Layout/Children.h"
+#include "Layout/Margin.h"
+#include "Layout/Visibility.h"
+#include "Misc/Attribute.h"
+#include "Misc/CommandLine.h"
+#include "Misc/Parse.h"
+#include "Misc/Paths.h"
+#include "SlotBase.h"
+#include "Styling/AppStyle.h"
+#include "Styling/SlateTypes.h"
+#include "TranslationDataManager.h"
+#include "TranslationPickerWidget.h"
+#include "TranslationUnit.h"
+#include "Types/SlateEnums.h"
+#include "Widgets/Input/SButton.h"
+#include "Widgets/Input/SCheckBox.h"
+#include "Widgets/Input/SEditableTextBox.h"
+#include "Widgets/Input/SMultiLineEditableTextBox.h"
+#include "Widgets/Layout/SBorder.h"
+#include "Widgets/Layout/SBox.h"
+#include "Widgets/Layout/SGridPanel.h"
+#include "Widgets/Layout/SScrollBox.h"
+#include "Widgets/Layout/SUniformGridPanel.h"
+#include "Widgets/SBoxPanel.h"
+#include "Widgets/SWindow.h"
+#include "Widgets/Text/STextBlock.h"
+
+struct FGeometry;
+
+#define LOCTEXT_NAMESPACE "TranslationPicker"
+
+TSharedPtr<FTranslationPickerSettingsManager> FTranslationPickerSettingsManager::TranslationPickerSettingsManagerInstance;
+
+// Default dimensions of the Translation Picker edit window (floating window also uses these sizes, so it matches roughly)
+const int32 STranslationPickerEditWindow::DefaultEditWindowWidth = 500;
+const int32 STranslationPickerEditWindow::DefaultEditWindowHeight = 500;
+
+class FTranslationPickerEditInputProcessor : public IInputProcessor
+{
+public:
+	FTranslationPickerEditInputProcessor(STranslationPickerEditWindow* InOwner)
+		: Owner(InOwner)
+	{
+	}
+
+	void SetOwner(STranslationPickerEditWindow* InOwner)
+	{
+		Owner = InOwner;
+	}
+
+	virtual ~FTranslationPickerEditInputProcessor() = default;
+
+	virtual void Tick(const float DeltaTime, FSlateApplication& SlateApp, TSharedRef<ICursor> Cursor) override
+	{
+	}
+
+	virtual bool HandleKeyDownEvent(FSlateApplication& SlateApp, const FKeyEvent& InKeyEvent) override
+	{
+		if (!Owner)
+		{
+			return false;
+		}
+
+		FKey Key = InKeyEvent.GetKey();
+
+		if (Key == EKeys::Escape)
+		{
+			Owner->Close();
+			return true;
+		}
+		else if (Key == EKeys::Enter)
+		{
+			Owner->RestorePicker();
+			return true;
+		}
+
+		return false;
+	}
+
+	virtual const TCHAR* GetDebugName() const override { return TEXT("TranslationPicker"); }
+
+private:
+	STranslationPickerEditWindow* Owner;
+};
+
+UTranslationPickerSettings::UTranslationPickerSettings(const FObjectInitializer& ObjectInitializer)
+	: Super(ObjectInitializer)
+{
+}
+
+void STranslationPickerEditWindow::Construct(const FArguments& InArgs)
+{
+	ParentWindow = InArgs._ParentWindow;
+	PickedTexts = InArgs._PickedTexts;
+	WindowContents = SNew(SBox);
+	TSharedRef<SVerticalBox> TextsBox = SNew(SVerticalBox);
+	UTranslationPickerSettings* TranslationPickerSettings = FTranslationPickerSettingsManager::Get()->GetSettings();
+
+	bool bShowLocServiceCheckbox = ILocalizationServiceModule::Get().GetProvider().IsEnabled();
+
+	if (!FParse::Param(FCommandLine::Get(), TEXT("AllowTranslationPickerSubmissionsToOneSky")))
+	{
+		bShowLocServiceCheckbox = false;
+		TranslationPickerSettings->bSubmitTranslationPickerChangesToLocalizationService = false;
+	}
+
+	// Add a new Translation Picker Edit Widget for each picked text
+	for (FText PickedText : PickedTexts)
+	{
+		TSharedRef<STranslationPickerEditWidget> NewEditWidget = 
+			SNew(STranslationPickerEditWidget)
+			.PickedText(PickedText)
+			.bAllowEditing(true);
+
+		EditWidgets.Add(NewEditWidget);
+
+		TextsBox->AddSlot()
+			.AutoHeight()
+			.Padding(FMargin(5))
+			[
+				SNew(SBorder)
+				[
+					NewEditWidget
+				]
+			];
+	}
+
+	TSharedPtr<SEditableTextBox> TextBox;
+	float DefaultPadding = 0.0f;
+
+	// Layout the Translation Picker Edit Widgets and some save/close buttons below them
+	WindowContents->SetContent(
+		SNew(SBorder)
+		.BorderImage(FAppStyle::GetBrush("ToolPanel.GroupBorder"))
+		[
+			SNew(SVerticalBox)
+
+			+SVerticalBox::Slot()
+			[
+				SNew(SScrollBox)
+				+SScrollBox::Slot()
+				.Padding(FMargin(8, 5, 8, 5))
+				[
+					TextsBox
+				]
+			]
+			
+			+SVerticalBox::Slot()
+			.AutoHeight()
+			.HAlign(HAlign_Fill)
+			.Padding(DefaultPadding)
+			[
+				SNew(SVerticalBox)
+				
+				+SVerticalBox::Slot()
+				.AutoHeight()
+				.HAlign(HAlign_Left)
+				.Padding(DefaultPadding)
+				[
+					SNew(SHorizontalBox)
+					.Visibility(bShowLocServiceCheckbox ? EVisibility::Visible : EVisibility::Collapsed)
+					
+					+SHorizontalBox::Slot()
+					.Padding(FMargin(3, 3, 3, 3))
+					.VAlign(VAlign_Center)
+					.AutoWidth()
+					[
+						SNew(SCheckBox)
+						.HAlign(HAlign_Center)
+						.IsChecked(TranslationPickerSettings->bSubmitTranslationPickerChangesToLocalizationService ? ECheckBoxState::Checked : ECheckBoxState::Unchecked)
+						.ToolTipText(LOCTEXT("SubmitTranslationPickerChangesToLocalizationServiceToolTip", "Submit changes to localization service"))
+						.OnCheckStateChanged_Lambda([&](ECheckBoxState CheckedState)
+						{
+							UTranslationPickerSettings* TranslationPickerSettingsLocal = FTranslationPickerSettingsManager::Get()->GetSettings();
+							TranslationPickerSettingsLocal->bSubmitTranslationPickerChangesToLocalizationService = CheckedState == ECheckBoxState::Checked;
+							TranslationPickerSettingsLocal->SaveConfig();
+						}
+						)
+					]
+					
+					+SHorizontalBox::Slot()
+					.Padding(FMargin(0, 0, 3, 0))
+					.VAlign(VAlign_Center)
+					.AutoWidth()
+					[
+						SNew(STextBlock)
+						.Text(LOCTEXT("SubmitTranslationPickerChangesToLocalizationService", "Save to Localization Service"))
+						.ToolTipText(LOCTEXT("SubmitTranslationPickerChangesToLocalizationServiceToolTip", "Submit changes to localization service"))
+					]
+				]
+				
+				+SVerticalBox::Slot()
+				.AutoHeight()
+				.HAlign(HAlign_Right)
+				.Padding(FMargin(0, 5))
+				[
+					SNew(SUniformGridPanel)
+					.SlotPadding(FAppStyle::GetMargin("StandardDialog.SlotPadding"))
+					.MinDesiredSlotWidth(FAppStyle::GetFloat("StandardDialog.MinDesiredSlotWidth"))
+					.MinDesiredSlotHeight(FAppStyle::GetFloat("StandardDialog.MinDesiredSlotHeight"))
+					
+					+SUniformGridPanel::Slot(0, 0)
+					[
+						SNew(SButton)
+						.HAlign(HAlign_Center)
+						.VAlign(VAlign_Center)
+						.ContentPadding(FAppStyle::GetMargin("StandardDialog.ContentPadding"))
+						.OnClicked(this, &STranslationPickerEditWindow::SaveAllAndClose)
+						.Text(LOCTEXT("SaveAllAndClose", "Save All and Close"))
+#if WITH_EDITOR
+						.Visibility(EVisibility::Visible)
+#else
+						.Visibility(EVisibility::Hidden)
+#endif // WITH_EDITOR
+					]
+					
+					+SUniformGridPanel::Slot(1, 0)
+					[
+						SNew(SButton)
+						.HAlign(HAlign_Center)
+						.VAlign(VAlign_Center)
+						.ContentPadding(FAppStyle::GetMargin("StandardDialog.ContentPadding"))
+						.OnClicked(this, &STranslationPickerEditWindow::Close)
+						.Text(LOCTEXT("CancelButton", "Cancel"))
+					]
+				]
+			]
+		]
+	);
+
+	ChildSlot
+	[
+		WindowContents.ToSharedRef()
+	];
+
+	InputProcessor = MakeShared<FTranslationPickerEditInputProcessor>(this);
+	FSlateApplication::Get().RegisterInputPreProcessor(InputProcessor, 0);
+}
+
+STranslationPickerEditWindow::~STranslationPickerEditWindow()
+{
+	if (InputProcessor.IsValid())
+	{
+		InputProcessor->SetOwner(nullptr);
+		if (FSlateApplication::IsInitialized())
+		{
+			FSlateApplication::Get().UnregisterInputPreProcessor(InputProcessor);
+		}
+		InputProcessor.Reset();
+	}
+}
+
+FReply STranslationPickerEditWindow::Close()
+{
+	if (ParentWindow.IsValid())
+	{
+		FSlateApplication::Get().RequestDestroyWindow(ParentWindow.Pin().ToSharedRef());
+		ParentWindow.Reset();
+	}
+
+	return FReply::Handled();
+}
+
+FReply STranslationPickerEditWindow::RestorePicker()
+{
+	Close();
+
+	TranslationPickerManager::OpenPickerWindow();
+
+	return FReply::Handled();
+}
+
+FReply STranslationPickerEditWindow::SaveAllAndClose()
+{
+	TArray<UTranslationUnit*> TempArray;
+
+	for (TSharedRef<STranslationPickerEditWidget> EditWidget : EditWidgets)
+	{
+		UTranslationUnit* TranslationUnit = EditWidget->GetTranslationUnitWithAnyChanges();
+		if (TranslationUnit != nullptr && EditWidget->CanSave())
+		{
+			TempArray.Add(TranslationUnit);
+		}
+	}
+
+	if (TempArray.Num() > 0)
+	{
+		UTranslationPickerSettings* TranslationPickerSettings = FTranslationPickerSettingsManager::Get()->GetSettings();
+		// Save the data via translation data manager
+		FTranslationDataManager::SaveSelectedTranslations(TempArray, ILocalizationServiceModule::Get().GetProvider().IsEnabled() && TranslationPickerSettings->bSubmitTranslationPickerChangesToLocalizationService);
+	}
+
+	Close();
+
+	return FReply::Handled();
+}
+
+void STranslationPickerEditWidget::Construct(const FArguments& InArgs)
+{
+	PickedText = InArgs._PickedText;
+	bAllowEditing = InArgs._bAllowEditing;
+
+	// Try and get the localization information for this text
+	FTextId TextId;
+	FString SourceString;
+	FString TranslationString;
+	{
+		if (const FString* SourceStringPtr = FTextInspector::GetSourceString(PickedText))
+		{
+			SourceString = *SourceStringPtr;
+		}
+		TranslationString = FTextInspector::GetDisplayString(PickedText);
+		TextId = FTextInspector::GetTextId(PickedText);
+	}
+
+	// Try and find the LocRes the active translation came from
+	// We assume the LocRes is named the same as the localization target
+	FString LocResPath;
+	FString LocTargetName;
+	FString LocResCultureName;
+#if WITH_EDITORONLY_DATA
+	if (!TextId.IsEmpty() && FTextLocalizationManager::Get().GetLocResID(TextId.GetNamespace(), TextId.GetKey(), LocResPath))
+	{
+		LocTargetName = FPaths::GetBaseFilename(LocResPath);
+
+		const FString CultureFilePath = FPaths::GetPath(LocResPath);
+		LocResCultureName = FPaths::GetBaseFilename(CultureFilePath);
+	}
+#endif // WITH_EDITORONLY_DATA
+
+	// Clean the package localization ID from the namespace (to mirror what the text gatherer does when scraping for translation data)
+	FString CleanNamespace = TextNamespaceUtil::StripPackageNamespace(TextId.GetNamespace().GetChars());
+
+	// Save the necessary data in UTranslationUnit for later.  This is what we pass to TranslationDataManager to save our edits
+	TranslationUnit = NewObject<UTranslationUnit>();
+	TranslationUnit->Namespace = CleanNamespace;
+	TranslationUnit->Key = TextId.GetKey().GetChars();
+	TranslationUnit->Source = SourceString;
+	TranslationUnit->Translation = TranslationString;
+	TranslationUnit->LocresPath = LocResPath;
+
+#if WITH_EDITOR
+	// Can only save if we have have an identity and are in a known localization target file
+	bHasRequiredLocalizationInfoForSaving = !TextId.IsEmpty() && !LocTargetName.IsEmpty();
+
+	const FTextBlockStyle& BoldText = FAppStyle::Get().GetWidgetStyle<FTextBlockStyle>("RichTextBlock.Bold");
+#else
+	// Bold text not available in game, suppress errors on 'RichTextBlock.Bold'
+	FTextBlockStyle BoldText = FTextBlockStyle::GetDefault();
+	BoldText.SetColorAndOpacity(FLinearColor(FColorList::White));
+#endif // WITH_EDITOR
+
+	TSharedPtr<SGridPanel> GridPanel;
+
+	// Layout all our data
+	ChildSlot
+	[
+		SNew(SHorizontalBox)
+		
+		+SHorizontalBox::Slot()
+		.FillWidth(1)
+		.Padding(FMargin(5))
+		[
+			SNew(SVerticalBox)
+			+SVerticalBox::Slot()
+			[
+				SAssignNew(GridPanel, SGridPanel)
+				.FillColumn(1,1)
+				
+				+SGridPanel::Slot(0,0)
+				.Padding(FMargin(2.5))
+				.HAlign(HAlign_Right)
+				[
+					SNew(STextBlock)
+					.TextStyle(&BoldText)
+					.Text(LOCTEXT("SourceLabel", "Source:"))
+				]
+
+				+SGridPanel::Slot(0, 1)
+				.Padding(FMargin(2.5))
+				.HAlign(HAlign_Right)
+				[
+					SNew(SBox)
+					// Hide translation if we don't have necessary information to modify
+					.Visibility(!bHasRequiredLocalizationInfoForSaving ? EVisibility::Collapsed : EVisibility::Visible)
+					[
+						SNew(STextBlock)
+						.TextStyle(&BoldText)
+#if WITH_EDITORONLY_DATA
+						.Text(FText::Format(LOCTEXT("TranslationLabelWithCulture", "Translation ({0}):"), FText::AsCultureInvariant(LocResCultureName)))
+#else
+						.Text(LOCTEXT("TranslationLabel", "Translation:"))
+#endif
+					]
+				]
+				
+				+SGridPanel::Slot(1, 0)
+				.Padding(FMargin(2.5))
+				[
+					SNew(SMultiLineEditableTextBox)
+					.IsReadOnly(true)
+					.Text(FText::AsCultureInvariant(SourceString))
+				]
+
+				+SGridPanel::Slot(1, 1)
+				.Padding(FMargin(2.5))
+				[
+					SNew(SBox)
+					// Hide translation if we don't have necessary information to modify
+					.Visibility(!bHasRequiredLocalizationInfoForSaving ? EVisibility::Collapsed : EVisibility::Visible)
+					[
+						SAssignNew(TextBox, SMultiLineEditableTextBox)
+						.IsReadOnly(!bAllowEditing || !bHasRequiredLocalizationInfoForSaving)
+						.Text(FText::AsCultureInvariant(TranslationString))
+						.HintText(LOCTEXT("TranslationEditTextBox_HintText", "Enter/edit translation here."))
+					]
+				]
+			]
+		]
+	];
+
+	if (!TextId.IsEmpty())
+	{
+		GridPanel->AddSlot(0, 2)
+			.Padding(FMargin(2.5))
+			.HAlign(HAlign_Right)
+			[
+				SNew(STextBlock)
+				.TextStyle(&BoldText)
+				.Text(LOCTEXT("NamespaceLabel", "Namespace:"))
+			];
+		GridPanel->AddSlot(1, 2)
+			.Padding(FMargin(2.5))
+			[
+				SNew(SEditableTextBox)
+				.IsReadOnly(true)
+				.Text(FText::AsCultureInvariant(CleanNamespace))
+			];
+		GridPanel->AddSlot(0, 3)
+			.Padding(FMargin(2.5))
+			.HAlign(HAlign_Right)
+			[
+				SNew(STextBlock)
+				.TextStyle(&BoldText)
+				.Text(LOCTEXT("KeyLabel", "Key:"))
+			];
+		GridPanel->AddSlot(1, 3)
+			.Padding(FMargin(2.5))
+			[
+				SNew(SEditableTextBox)
+				.IsReadOnly(true)
+				.Text(FText::AsCultureInvariant(TextId.GetKey().GetChars()))
+			];
+		
+		int32 Row = 4;
+		if (bHasRequiredLocalizationInfoForSaving)
+		{
+
+#if WITH_EDITOR
+			GridPanel->AddSlot(0, Row)
+				.Padding(FMargin(2.5))
+				.HAlign(HAlign_Right)
+				[
+					SNew(STextBlock)
+					.TextStyle(&BoldText)
+					.Text(LOCTEXT("LocresFileLabel", "Target:"))
+				];
+			GridPanel->AddSlot(1, Row)
+				.Padding(FMargin(2.5))
+				[
+					SNew(SEditableTextBox)
+					.IsReadOnly(true)
+					.Text(FText::AsCultureInvariant(LocTargetName))
+				];
+			++Row;
+#endif // WITH_EDITOR
+
+			GridPanel->AddSlot(0, Row)
+				.Padding(FMargin(2.5))
+				[
+					SNew(SButton)
+					.HAlign(HAlign_Center)
+					.VAlign(VAlign_Center)
+					.ContentPadding(FAppStyle::GetMargin("StandardDialog.ContentPadding"))
+					.OnClicked(this, &STranslationPickerEditWidget::CopyNamespaceAndKey)
+					.Visibility(bAllowEditing ? EVisibility::Visible : EVisibility::Collapsed)
+					.Text(LOCTEXT("CopyNamespaceAndKey", "Copy Namespace,Key"))
+				];
+			GridPanel->AddSlot(1, Row)
+				.Padding(FMargin(2.5))
+				.HAlign(HAlign_Right)
+				[
+					SNew(SButton)
+					.HAlign(HAlign_Center)
+					.VAlign(VAlign_Center)
+					.ContentPadding(FAppStyle::GetMargin("StandardDialog.ContentPadding"))
+					.OnClicked(this, &STranslationPickerEditWidget::SaveAndPreview)
+					.IsEnabled(bHasRequiredLocalizationInfoForSaving)
+					.Visibility(bAllowEditing ? EVisibility::Visible : EVisibility::Collapsed)
+					.Text(bHasRequiredLocalizationInfoForSaving ? LOCTEXT("SaveAndPreviewButtonText", "Save and Preview") : LOCTEXT("SaveAndPreviewButtonDisabledText", "Cannot Save"))
+				];
+		}
+		else
+		{
+			GridPanel->AddSlot(0, Row)
+				.Padding(FMargin(2.5))
+				.ColumnSpan(2)
+				.HAlign(HAlign_Center)
+				[
+					SNew(STextBlock)
+					.Text(LOCTEXT("TextLocalizable_RequiresGather", "This text is localizable (requires gather)."))
+				];
+		}
+	}
+	else
+	{
+		FText TextNotLocalizableReason = LOCTEXT("TextNotLocalizable_Generic", "This text is not localizable.");
+		if (PickedText.IsCultureInvariant())
+		{
+			TextNotLocalizableReason = LOCTEXT("TextNotLocalizable_CultureInvariant", "This text is not localizable (culture-invariant).");
+		}
+		else if (PickedText.IsTransient())
+		{
+			TextNotLocalizableReason = LOCTEXT("TextNotLocalizable_Transient", "This text is not localizable (transient).");
+		}
+		else if (!PickedText.ShouldGatherForLocalization())
+		{
+			TextNotLocalizableReason = LOCTEXT("TextNotLocalizable_InvalidForGather", "This text is not localizable (invalid for gather).");
+		}
+
+		GridPanel->AddSlot(0, 2)
+			.Padding(FMargin(2.5))
+			.ColumnSpan(2)
+			.HAlign(HAlign_Center)
+			[
+				SNew(STextBlock)
+				.Text(TextNotLocalizableReason)
+			];
+	}
+}
+
+void STranslationPickerEditWidget::AddReferencedObjects(FReferenceCollector& Collector)
+{
+	Collector.AddReferencedObject(TranslationUnit);
+}
+
+FReply STranslationPickerEditWidget::SaveAndPreview()
+{
+	// Update translation string from entered text
+	TranslationUnit->Translation = TextBox->GetText().ToString();
+
+#if WITH_EDITOR
+	UTranslationPickerSettings* TranslationPickerSettings = FTranslationPickerSettingsManager::Get()->GetSettings();
+
+	// Save the data via translation data manager
+	TArray<UTranslationUnit*> TempArray;
+	TempArray.Add(TranslationUnit);
+	FTranslationDataManager::SaveSelectedTranslations(TempArray, ILocalizationServiceModule::Get().GetProvider().IsEnabled() && TranslationPickerSettings->bSubmitTranslationPickerChangesToLocalizationService);
+
+#else
+
+	FTextLocalizationManager::Get().ReplaceStringInLiveTable(TranslationUnit->Namespace, TranslationUnit->Key, TranslationUnit->Translation);
+
+	if (IConsoleObject* CObj = IConsoleManager::Get().FindConsoleObject(TEXT("Slate.TriggerInvalidate")))
+	{
+		CObj->AsCommand()->Execute(/*Args=*/TArray<FString>(), /*InWorld=*/nullptr, *GLog);
+	}
+#endif // WITH_EDITOR
+
+	return FReply::Handled();
+}
+
+FReply STranslationPickerEditWidget::CopyNamespaceAndKey()
+{
+	const FString CopyString = FString::Printf(TEXT("%s,%s"), *TranslationUnit->Namespace, *TranslationUnit->Key);
+	
+	FPlatformApplicationMisc::ClipboardCopy(*CopyString);
+
+	UE_LOG(LogConsoleResponse, Display, TEXT("Copied Namespace,Key to clipboard: %s"), *CopyString);
+
+	return FReply::Handled();
+}
+
+UTranslationUnit* STranslationPickerEditWidget::GetTranslationUnitWithAnyChanges()
+{
+	if (TranslationUnit)
+	{
+		// Update translation string from entered text
+		TranslationUnit->Translation = TextBox->GetText().ToString();
+
+		return TranslationUnit;
+	}
+
+	return nullptr;
+}
+
+#undef LOCTEXT_NAMESPACE
diff --git a/Engine/Source/Developer/TranslationEditor/Private/TranslationPickerEditWindow.h b/Engine/Source/Developer/TranslationEditor/Private/TranslationPickerEditWindow.h
new file mode 100644
--- /dev/null
+++ b/Engine/Source/Developer/TranslationEditor/Private/TranslationPickerEditWindow.h
@@ -0,0 +1,195 @@
+// Copyright Epic Games, Inc. All Rights Reserved.
+
+#pragma once
+
+#include "Containers/Array.h"
+#include "Containers/UnrealString.h"
+#include "HAL/Platform.h"
+#include "Input/Reply.h"
+#include "Internationalization/Text.h"
+#include "Templates/SharedPointer.h"
+#include "UObject/GCObject.h"
+#include "UObject/Object.h"
+#include "UObject/UObjectGlobals.h"
+#include "Widgets/DeclarativeSyntaxSupport.h"
+#include "Widgets/SCompoundWidget.h"
+
+#include "TranslationPickerEditWindow.generated.h"
+
+struct FGeometry;
+struct FKeyEvent;
+
+#define LOCTEXT_NAMESPACE "TranslationPicker"
+
+class FTranslationPickerEditInputProcessor;
+class SBox;
+class SMultiLineEditableTextBox;
+class SWindow;
+class UTranslationUnit;
+
+UCLASS(config = TranslationPickerSettings)
+class UTranslationPickerSettings : public UObject
+{
+	GENERATED_UCLASS_BODY()
+
+public:
+	/** Bool submit translation picker changes to Localization Service */
+	UPROPERTY(config)
+	bool bSubmitTranslationPickerChangesToLocalizationService;
+};
+
+class FTranslationPickerSettingsManager
+{
+	FTranslationPickerSettingsManager()
+	{
+		TranslationPickerSettingsObject = NewObject<UTranslationPickerSettings>();
+		TranslationPickerSettingsObject->LoadConfig();
+	}
+
+public:
+	void SaveSettings()
+	{
+		TranslationPickerSettingsObject->SaveConfig();
+	}
+
+	void LoadSettings()
+	{
+		TranslationPickerSettingsObject->LoadConfig();
+	}
+
+	UTranslationPickerSettings* GetSettings()
+	{
+		return TranslationPickerSettingsObject;
+	}
+
+	/**
+	* Gets a reference to the translation picker manager instance.
+	*
+	* @return A reference to the translation picker manager instance.
+	*/
+	static inline TSharedPtr<FTranslationPickerSettingsManager>& Get()
+	{
+		if (!TranslationPickerSettingsManagerInstance.IsValid())
+		{
+			TranslationPickerSettingsManagerInstance = MakeShareable(new FTranslationPickerSettingsManager());
+		}
+
+		return TranslationPickerSettingsManagerInstance;
+	}
+
+private:
+
+	static TSharedPtr<FTranslationPickerSettingsManager> TranslationPickerSettingsManagerInstance;
+
+	/** Used to load and store settings for the Translation Picker */
+	UTranslationPickerSettings* TranslationPickerSettingsObject;
+};
+
+/** Translation picker edit Widget to handle the display and editing of a single selected FText */
+class STranslationPickerEditWidget : public SCompoundWidget, public FGCObject
+{
+	SLATE_BEGIN_ARGS(STranslationPickerEditWidget) {}
+
+	SLATE_ARGUMENT(FText, PickedText)
+
+	SLATE_ARGUMENT(bool, bAllowEditing)
+
+	SLATE_END_ARGS()
+
+	void Construct(const FArguments& InArgs);
+
+	virtual void AddReferencedObjects(FReferenceCollector& Collector) override;
+	virtual FString GetReferencerName() const override
+	{
+		return TEXT("STranslationPickerEditWidget");
+	}
+
+	/** Return the translation unit for this text, with any modifications */
+	UTranslationUnit* GetTranslationUnitWithAnyChanges();
+
+	/** Whether or not we can save */
+	bool CanSave()
+	{
+		return bAllowEditing && bHasRequiredLocalizationInfoForSaving;
+	}
+
+private:
+
+	virtual bool SupportsKeyboardFocus() const override
+	{
+		return true;
+	}
+
+	FReply SaveAndPreview();
+
+	FReply CopyNamespaceAndKey();
+
+	/** The FText that we are using this widget to translate */
+	FText PickedText;
+
+	/** The translation we're editing represented in a UTranslationUnit object */
+	TObjectPtr<UTranslationUnit> TranslationUnit;
+
+	/** The text box for entering/modifying a translation */
+	TSharedPtr<SMultiLineEditableTextBox> TextBox;
+
+	/** Whether or not to show the save button*/
+	bool bAllowEditing;
+
+	/** Whether or not we were able to find the necessary info for saving */
+	bool bHasRequiredLocalizationInfoForSaving = true;
+};
+
+/** Translation picker edit window to allow you to translate selected FTexts in place */
+class STranslationPickerEditWindow : public SCompoundWidget
+{
+	SLATE_BEGIN_ARGS(STranslationPickerEditWindow) {}
+
+	SLATE_ARGUMENT(TWeakPtr<SWindow>, ParentWindow)
+
+	SLATE_ARGUMENT(TArray<FText>, PickedTexts)
+
+	SLATE_END_ARGS()
+
+	virtual ~STranslationPickerEditWindow();
+
+	void Construct(const FArguments& InArgs);
+
+	// Default dimensions of the Translation Picker edit window (floating window also uses these sizes, so it matches roughly)
+	static const int32 DefaultEditWindowWidth;
+	static const int32 DefaultEditWindowHeight;
+
+private:
+	friend class FTranslationPickerEditInputProcessor;
+
+	FReply Close();
+
+	/** We need to support keyboard focus to process the 'Esc' key */
+	virtual bool SupportsKeyboardFocus() const override
+	{
+		return true;
+	}
+
+	/** Input processor used to capture the 'Esc' key */
+	TSharedPtr<FTranslationPickerEditInputProcessor> InputProcessor;
+
+	/** Handle to the window that contains this widget */
+	TWeakPtr<SWindow> ParentWindow;
+
+	/** Contents of the window */
+	TSharedPtr<SBox> WindowContents;
+
+	/** The FTexts that we have found under the cursor */
+	TArray<FText> PickedTexts;
+
+	/** All of our current edit widgets */
+	TArray<TSharedRef<STranslationPickerEditWidget>> EditWidgets;
+
+	/** Return to picker floating window */
+	FReply RestorePicker();
+
+	/** Save all translations and close */
+	FReply SaveAllAndClose();
+};
+
+#undef LOCTEXT_NAMESPACE
diff --git a/Engine/Source/Developer/TranslationEditor/Private/TranslationPickerFloatingWindow.cpp b/Engine/Source/Developer/TranslationEditor/Private/TranslationPickerFloatingWindow.cpp
new file mode 100644
--- /dev/null
+++ b/Engine/Source/Developer/TranslationEditor/Private/TranslationPickerFloatingWindow.cpp
@@ -0,0 +1,409 @@
+// Copyright Epic Games, Inc. All Rights Reserved.
+
+#include "TranslationPickerFloatingWindow.h"
+
+#include "Containers/SortedMap.h"
+#include "Containers/UnrealString.h"
+#include "Framework/Application/IInputProcessor.h"
+#include "Framework/Application/SlateApplication.h"
+#include "Framework/Docking/TabManager.h"
+#include "Framework/Text/TextLayout.h"
+#include "HAL/Platform.h"
+#include "HAL/PlatformCrt.h"
+#include "HAL/PlatformApplicationMisc.h"
+#include "Input/Events.h"
+#include "InputCoreTypes.h"
+#include "Internationalization/Internationalization.h"
+#include "Internationalization/Text.h"
+#include "Internationalization/TextNamespaceUtil.h"
+#include "Layout/ArrangedChildren.h"
+#include "Layout/ArrangedWidget.h"
+#include "Layout/BasicLayoutWidgetSlot.h"
+#include "Layout/Children.h"
+#include "Layout/ChildrenBase.h"
+#include "Layout/Margin.h"
+#include "Layout/SlateRect.h"
+#include "Math/Vector2D.h"
+#include "Misc/Attribute.h"
+#if WITH_EDITOR
+#include "SDocumentationToolTip.h"
+#endif // WITH_EDITOR
+#include "SlotBase.h"
+#include "TranslationPickerEditWindow.h"
+#include "TranslationPickerWidget.h"
+#include "Types/SlateEnums.h"
+#include "Widgets/IToolTip.h"
+#include "Widgets/Input/SEditableText.h"
+#include "Widgets/Layout/SBorder.h"
+#include "Widgets/Layout/SScrollBox.h"
+#include "Widgets/SBoxPanel.h"
+#include "Widgets/SToolTip.h"
+#include "Widgets/SWidget.h"
+#include "Widgets/SWindow.h"
+#include "Widgets/Text/SMultiLineEditableText.h"
+#include "Widgets/Text/SRichTextBlock.h"
+#include "Widgets/Text/STextBlock.h"
+
+class ICursor;
+struct FGeometry;
+
+#define LOCTEXT_NAMESPACE "TranslationPicker"
+
+class FTranslationPickerInputProcessor : public IInputProcessor
+{
+public:
+	FTranslationPickerInputProcessor(STranslationPickerFloatingWindow* InOwner)
+		: Owner(InOwner)
+	{
+	}
+
+	void SetOwner(STranslationPickerFloatingWindow* InOwner)
+	{
+		Owner = InOwner;
+	}
+
+	virtual ~FTranslationPickerInputProcessor() = default;
+
+	virtual void Tick(const float DeltaTime, FSlateApplication& SlateApp, TSharedRef<ICursor> Cursor) override
+	{
+	}
+
+	virtual bool HandleKeyDownEvent(FSlateApplication& SlateApp, const FKeyEvent& InKeyEvent) override
+	{
+		if (!Owner)
+		{
+			return false;
+		}
+
+		FKey Key = InKeyEvent.GetKey();
+
+		if (Key == EKeys::Escape)
+		{
+			TranslationPickerManager::ClosePickerWindow();
+			return true;
+		}
+		else if (Key == EKeys::Enter)
+		{
+			Owner->SwitchToEditWindow();
+			return true;
+		}
+		else if (InKeyEvent.IsControlDown())
+		{
+			const uint32* KeyCode = nullptr;
+			const uint32* CharCode = nullptr;
+			FInputKeyManager::Get().GetCodesFromKey(Key, KeyCode, CharCode);
+			if (CharCode == nullptr)
+			{
+				return false;
+			}
+
+			const uint32* KeyCodeOne = nullptr;
+			const uint32* CharCodeOne = nullptr;
+			FInputKeyManager::Get().GetCodesFromKey(EKeys::One, KeyCodeOne, CharCodeOne);
+
+			int32 EntryIndex = *CharCode - *CharCodeOne;
+			if (EntryIndex < 0 || EntryIndex > 4 || EntryIndex >= Owner->PickedTexts.Num())
+			{
+				return false;	// Handle only first five entries, the max number of entries that fit in the floating picker
+			}
+
+			const FText& PickedText = Owner->PickedTexts[EntryIndex];
+			
+			FTextId TextId = FTextInspector::GetTextId(PickedText);
+
+			// Clean the package localization ID from the namespace (to mirror what the text gatherer does when scraping for translation data)
+			FString EntryNamespace = TextNamespaceUtil::StripPackageNamespace(TextId.GetNamespace().GetChars());
+			FString EntryKey = TextId.GetKey().GetChars();
+
+			const FString CopyString = FString::Printf(TEXT("%s,%s"), *EntryNamespace, *EntryKey);
+	
+			FPlatformApplicationMisc::ClipboardCopy(*CopyString);
+
+			UE_LOG(LogConsoleResponse, Display, TEXT("Copied Namespace,Key to clipboard: %s"), *CopyString);
+
+			return true;
+		}
+
+		return false;
+	}
+
+	virtual const TCHAR* GetDebugName() const override { return TEXT("TranslationPicker"); }
+
+private:
+	STranslationPickerFloatingWindow* Owner;
+};
+
+void STranslationPickerFloatingWindow::Construct(const FArguments& InArgs)
+{
+	ParentWindow = InArgs._ParentWindow;
+	WindowContents = SNew(SToolTip);
+
+	ChildSlot
+	[
+		WindowContents.ToSharedRef()
+	];
+
+	InputProcessor = MakeShared<FTranslationPickerInputProcessor>(this);
+	FSlateApplication::Get().RegisterInputPreProcessor(InputProcessor, 0);
+}
+
+STranslationPickerFloatingWindow::~STranslationPickerFloatingWindow()
+{
+	if (InputProcessor.IsValid())
+	{
+		InputProcessor->SetOwner(nullptr);
+		if (FSlateApplication::IsInitialized())
+		{
+			FSlateApplication::Get().UnregisterInputPreProcessor(InputProcessor);
+		}
+		InputProcessor.Reset();
+	}
+}
+
+void STranslationPickerFloatingWindow::Tick( const FGeometry& AllottedGeometry, const double InCurrentTime, const float InDeltaTime )
+{
+	FWidgetPath Path = FSlateApplication::Get().LocateWindowUnderMouse(FSlateApplication::Get().GetCursorPos(), FSlateApplication::Get().GetInteractiveTopLevelWindows(), true);
+
+	if (Path.IsValid())
+	{
+		// If the path of widgets we're hovering over changed since last time (or if this is the first tick and LastTickHoveringWidgetPath hasn't been set yet)
+		if (!LastTickHoveringWidgetPath.IsValid() || LastTickHoveringWidgetPath.ToWidgetPath().ToString() != Path.ToString())
+		{
+			// Clear all previous text and widgets
+			PickedTexts.Reset();
+
+			// Process the leaf-widget under the cursor
+			if (Path.Widgets.Num() > 0)
+			{
+				// General Widget case
+				TSharedRef<SWidget> PathWidget = Path.Widgets.Last().Widget;
+				PickTextFromWidget(PathWidget);
+
+				// Tooltip case
+				TSharedPtr<IToolTip> Tooltip = PathWidget->GetToolTip();
+				if (Tooltip.IsValid() && !Tooltip->IsEmpty())
+				{
+					PickTextFromWidget(Tooltip->AsWidget());
+				}
+
+				// Also include tooltips from parent widgets in this path (since they may be visible)
+				for (int32 ParentPathIndex = Path.Widgets.Num() - 2; ParentPathIndex >= 0; --ParentPathIndex)
+				{
+					TSharedRef<SWidget> ParentPathWidget = Path.Widgets[ParentPathIndex].Widget;
+
+					// Tooltip case
+					TSharedPtr<IToolTip> ParentTooltip = ParentPathWidget->GetToolTip();
+					if (ParentTooltip.IsValid() && !ParentTooltip->IsEmpty())
+					{
+						PickTextFromWidget(ParentTooltip->AsWidget());
+					}
+				}
+			}
+
+			TSharedRef<SVerticalBox> TextsBox = SNew(SVerticalBox);
+
+			// Add a new Translation Picker Edit Widget for each picked text
+			for (FText PickedText : PickedTexts)
+			{
+				TextsBox->AddSlot()
+					.AutoHeight()
+					.Padding(FMargin(5))
+					[
+						SNew(SBorder)
+						[
+							SNew(STranslationPickerEditWidget)
+							.PickedText(PickedText)
+							.bAllowEditing(false)
+						]
+					];
+			}
+
+			WindowContents->SetContentWidget(
+				SNew(SVerticalBox)
+				
+				+SVerticalBox::Slot()
+				.Padding(0)
+				.FillHeight(1)
+				.Padding(FMargin(5))
+				[
+					SNew(SScrollBox)
+					.Orientation(EOrientation::Orient_Vertical)
+					.ScrollBarAlwaysVisible(true)
+					
+					+SScrollBox::Slot()
+					.Padding(FMargin(0))
+					[
+						TextsBox
+					]
+				]
+
+				+SVerticalBox::Slot()
+				.Padding(0)
+				.AutoHeight()
+				.Padding(FMargin(5))
+				[
+					SNew(STextBlock)
+					.Text(PickedTexts.Num() > 0 ? LOCTEXT("TranslationPickerEnterToEdit", "Press Enter to edit translations") : LOCTEXT("TranslationPickerHoverToViewEditEscToQuit", "Hover over text to view/edit translations, or press Esc to quit"))
+					.Justification(ETextJustify::Center)
+				]
+			);
+		}
+	}
+
+	if (ParentWindow.IsValid())
+	{
+		FVector2D WindowSize = ParentWindow.Pin()->GetSizeInScreen();
+		FVector2D DesiredPosition = FSlateApplication::Get().GetCursorPos();
+		DesiredPosition.X -= FSlateApplication::Get().GetCursorSize().X;
+		DesiredPosition.Y += FSlateApplication::Get().GetCursorSize().Y;
+
+		// Move to opposite side of the cursor than the tool tip, so they don't overlaps
+		DesiredPosition.X -= WindowSize.X;
+
+		// Clamp to work area
+		DesiredPosition = FSlateApplication::Get().CalculateTooltipWindowPosition(FSlateRect(DesiredPosition, DesiredPosition), WindowSize, false);
+
+		// also kind of a hack, but this is the only way at the moment to get a 'cursor decorator' without using the drag-drop code path
+		ParentWindow.Pin()->MoveWindowTo(DesiredPosition);
+	}
+
+	LastTickHoveringWidgetPath = FWeakWidgetPath(Path);
+}
+
+void STranslationPickerFloatingWindow::PickTextFromWidget(TSharedRef<SWidget> Widget)
+{
+	auto AppendPickedTextImpl = [this](const FText& InPickedText)
+	{
+		const bool bAlreadyPicked = PickedTexts.ContainsByPredicate([&InPickedText](const FText& InOtherPickedText)
+		{
+			return InOtherPickedText.IdenticalTo(InPickedText);
+		});
+
+		if (!bAlreadyPicked)
+		{
+			PickedTexts.Add(InPickedText);
+		}
+	};
+
+	auto AppendPickedText = [this, AppendPickedTextImpl](const FText& InPickedText)
+	{
+		if (InPickedText.IsEmpty())
+		{
+			return;
+		}
+
+		// Search the text from this widget's FText::Format history to find any source text
+		TArray<FHistoricTextFormatData> HistoricFormatData;
+		FTextInspector::GetHistoricFormatData(InPickedText, HistoricFormatData);
+
+		if (HistoricFormatData.Num() > 0)
+		{
+			for (const FHistoricTextFormatData& HistoricFormatDataItem : HistoricFormatData)
+			{
+				AppendPickedTextImpl(HistoricFormatDataItem.SourceFmt.GetSourceText());
+
+				for (auto It = HistoricFormatDataItem.Arguments.CreateConstIterator(); It; ++It)
+				{
+					const FFormatArgumentValue& ArgumentValue = It.Value();
+					if (ArgumentValue.GetType() == EFormatArgumentType::Text)
+					{
+						AppendPickedTextImpl(ArgumentValue.GetTextValue());
+					}
+				}
+			}
+		}
+		else
+		{
+			AppendPickedTextImpl(InPickedText);
+		}
+	};
+
+	// Have to parse the various widget types to find the FText
+	if (Widget->GetTypeAsString() == "STextBlock")
+	{
+		STextBlock& TextBlock = (STextBlock&)Widget.Get();
+		AppendPickedText(TextBlock.GetText());
+	}
+	else if (Widget->GetTypeAsString() == "SRichTextBlock")
+	{
+		SRichTextBlock& RichTextBlock = (SRichTextBlock&)Widget.Get();
+		AppendPickedText(RichTextBlock.GetText());
+	}
+	else if (Widget->GetTypeAsString() == "SToolTip")
+	{
+		SToolTip& ToolTipWidget = (SToolTip&)Widget.Get();
+		AppendPickedText(ToolTipWidget.GetTextTooltip());
+	}
+#if WITH_EDITOR
+	else if (Widget->GetTypeAsString() == "SDocumentationToolTip")
+	{
+		SDocumentationToolTip& DocumentationToolTip = (SDocumentationToolTip&)Widget.Get();
+		AppendPickedText(DocumentationToolTip.GetTextTooltip());
+	}
+#endif // WITH_EDITOR
+	else if (Widget->GetTypeAsString() == "SEditableText")
+	{
+		SEditableText& EditableText = (SEditableText&)Widget.Get();
+		AppendPickedText(EditableText.GetText());
+		AppendPickedText(EditableText.GetHintText());
+	}
+	else if (Widget->GetTypeAsString() == "SMultiLineEditableText")
+	{
+		SMultiLineEditableText& MultiLineEditableText = (SMultiLineEditableText&)Widget.Get();
+		AppendPickedText(MultiLineEditableText.GetText());
+		AppendPickedText(MultiLineEditableText.GetHintText());
+	}
+
+	// Recurse into child widgets
+	PickTextFromChildWidgets(Widget);
+}
+
+void STranslationPickerFloatingWindow::PickTextFromChildWidgets(TSharedRef<SWidget> Widget)
+{
+	FChildren* Children = Widget->GetChildren();
+
+	for (int32 ChildIndex = 0; ChildIndex < Children->Num(); ++ChildIndex)
+	{
+		TSharedRef<SWidget> ChildWidget = Children->GetChildAt(ChildIndex);
+
+		// Pull out any FText from this child widget
+		PickTextFromWidget(ChildWidget);
+	}
+}
+
+void STranslationPickerFloatingWindow::SwitchToEditWindow()
+{
+	if (PickedTexts.Num() > 0)
+	{
+		// Open a different window to allow editing of the translation
+		TSharedRef<SWindow> NewWindow = SNew(SWindow)
+			.Title(LOCTEXT("TranslationPickerEditWindowTitle", "Edit Translations"))
+			.CreateTitleBar(true)
+			.SizingRule(ESizingRule::UserSized);
+
+		TSharedRef<STranslationPickerEditWindow> EditWindow = SNew(STranslationPickerEditWindow)
+			.ParentWindow(NewWindow)
+			.PickedTexts(PickedTexts);
+
+		NewWindow->SetContent(EditWindow);
+
+		// Make this roughly the same size as the Edit Window, so when you press Esc to edit, the window is in basically the same size
+		NewWindow->Resize(FVector2D(STranslationPickerEditWindow::DefaultEditWindowWidth, STranslationPickerEditWindow::DefaultEditWindowHeight));
+
+		TSharedPtr<SWindow> RootWindow = FGlobalTabmanager::Get()->GetRootWindow();
+		if (RootWindow.IsValid())
+		{
+			FSlateApplication::Get().AddWindowAsNativeChild(NewWindow, RootWindow.ToSharedRef());
+		}
+		else
+		{
+			FSlateApplication::Get().AddWindow(NewWindow);
+		}
+
+		NewWindow->MoveWindowTo(ParentWindow.Pin()->GetPositionInScreen());
+	}
+
+	TranslationPickerManager::ClosePickerWindow();
+}
+
+#undef LOCTEXT_NAMESPACE
diff --git a/Engine/Source/Developer/TranslationEditor/Private/TranslationPickerFloatingWindow.h b/Engine/Source/Developer/TranslationEditor/Private/TranslationPickerFloatingWindow.h
new file mode 100644
--- /dev/null
+++ b/Engine/Source/Developer/TranslationEditor/Private/TranslationPickerFloatingWindow.h
@@ -0,0 +1,66 @@
+// Copyright Epic Games, Inc. All Rights Reserved.
+
+#pragma once
+
+#include "Containers/Array.h"
+#include "Layout/WidgetPath.h"
+#include "Templates/SharedPointer.h"
+#include "Widgets/DeclarativeSyntaxSupport.h"
+#include "Widgets/SCompoundWidget.h"
+
+class FText;
+class SWidget;
+class SWindow;
+struct FGeometry;
+
+#define LOCTEXT_NAMESPACE "TranslationPicker"
+
+class FTranslationPickerInputProcessor;
+class SToolTip;
+
+/** Translation picker floating window to show details of FText(s) under cursor, and allow in-place translation via TranslationPickerEditWindow */
+class STranslationPickerFloatingWindow : public SCompoundWidget
+{
+	SLATE_BEGIN_ARGS(STranslationPickerFloatingWindow) {}
+
+	SLATE_ARGUMENT(TWeakPtr<SWindow>, ParentWindow)
+
+	SLATE_END_ARGS()
+
+	virtual ~STranslationPickerFloatingWindow();
+
+	void Construct(const FArguments& InArgs);
+
+private:
+	friend class FTranslationPickerInputProcessor;
+
+	virtual void Tick( const FGeometry& AllottedGeometry, const double InCurrentTime, const float InDeltaTime ) override;
+
+	/** Pull the FText reference out of an SWidget */
+	void PickTextFromWidget(TSharedRef<SWidget> Widget);
+
+	/** Pull the FText reference out of the child widgets of an SWidget */
+	void PickTextFromChildWidgets(TSharedRef<SWidget> Widget);
+
+	/** Switch from floating window to edit window */
+	void SwitchToEditWindow();
+
+	/** Input processor used to capture the 'Esc' key */
+	TSharedPtr<FTranslationPickerInputProcessor> InputProcessor;
+
+	/** Handle to the window that contains this widget */
+	TWeakPtr<SWindow> ParentWindow;
+
+	/** Contents of the window */
+	TSharedPtr<SToolTip> WindowContents;
+
+	/** The FTexts that we have found under the cursor */
+	TArray<FText> PickedTexts;
+
+	/**
+	* The path widgets we were hovering over last tick
+	*/
+	FWeakWidgetPath LastTickHoveringWidgetPath;
+};
+
+#undef LOCTEXT_NAMESPACE
diff --git a/Engine/Source/Developer/TranslationEditor/Private/TranslationPickerWidget.cpp b/Engine/Source/Developer/TranslationEditor/Private/TranslationPickerWidget.cpp
new file mode 100644
--- /dev/null
+++ b/Engine/Source/Developer/TranslationEditor/Private/TranslationPickerWidget.cpp
@@ -0,0 +1,172 @@
+// Copyright Epic Games, Inc. All Rights Reserved.
+
+#include "TranslationPickerWidget.h"
+
+#include "Framework/Application/SlateApplication.h"
+#include "Framework/Docking/TabManager.h"
+#include "Internationalization/Internationalization.h"
+#include "Internationalization/Text.h"
+#include "Layout/Children.h"
+#include "Layout/Margin.h"
+#include "Math/MathFwd.h"
+#include "Math/Vector2D.h"
+#include "Misc/Attribute.h"
+#include "SlotBase.h"
+#include "Styling/AppStyle.h"
+#include "Styling/ISlateStyle.h"
+#include "Styling/SlateTypes.h"
+#include "TranslationPickerEditWindow.h"
+#include "TranslationPickerFloatingWindow.h"
+#include "Types/SlateEnums.h"
+#include "UObject/NameTypes.h"
+#include "Widgets/Images/SImage.h"
+#include "Widgets/Input/SCheckBox.h"
+#include "Widgets/SBoxPanel.h"
+#include "Widgets/SToolTip.h"
+#include "Widgets/SWindow.h"
+#include "Widgets/Text/STextBlock.h"
+
+class SWidget;
+
+#define LOCTEXT_NAMESPACE "TranslationPicker"
+
+
+void STranslationWidgetPicker::Construct(const FArguments& InArgs)
+{
+	// Mimicking a toolbar button look
+
+	// Icon for the picker widget button
+	TSharedRef< SWidget > IconWidget =
+		SNew(SImage)
+		.Image(FAppStyle::GetBrush("TranslationEditor.TranslationPicker"));
+
+	// Style settings
+	FName StyleSet = FAppStyle::GetAppStyleSetName();
+	FName StyleName = "Toolbar";
+
+	FText ToolTipText = LOCTEXT("TranslationPickerTooltip", "Open the Translation Picker");
+
+	// Create the content for our button
+	TSharedRef< SWidget > ButtonContent =
+
+		SNew(SHorizontalBox)
+
+		+ SHorizontalBox::Slot()
+		.FillWidth(1)
+		.VAlign(VAlign_Center)
+		[
+			SNew(SVerticalBox)
+
+			// Icon image
+			+ SVerticalBox::Slot()
+			.AutoHeight()
+			.HAlign(HAlign_Center)	// Center the icon horizontally, so that large labels don't stretch out the artwork
+			[
+				IconWidget
+			]
+
+			// Label text
+			+ SVerticalBox::Slot().AutoHeight()
+				.HAlign(HAlign_Center)	// Center the label text horizontally
+				[
+					SNew(STextBlock)
+					.Text(LOCTEXT("TranslationPicker", "Translation Picker"))
+					.TextStyle(FAppStyle::Get(), FName("ToolBar.Label"))	// Smaller font for tool tip labels
+					.ShadowOffset(FVector2D::UnitVector)
+				]
+		];
+
+	FName CheckboxStyle = ISlateStyle::Join(StyleName, ".SToolBarButtonBlock.CheckBox.Padding");
+
+	ChildSlot
+		[
+			// Create a check box
+			SNew(SCheckBox)
+
+			// Use the tool bar style for this check box
+			.Style(FAppStyle::Get(), "ToolBar.ToggleButton")
+
+			// User will have set the focusable attribute for the block, honor it
+			.IsFocusable(false)
+
+			// Pass along the block's tool-tip string
+			.ToolTip(SNew(SToolTip).Text(ToolTipText))
+			[
+				ButtonContent
+			]
+
+			// Bind the button's "on checked" event to our object's method for this
+			.OnCheckStateChanged(this, &STranslationWidgetPicker::OnCheckStateChanged)
+
+
+				// Bind the check box's "checked" state to our user interface action
+				.IsChecked(this, &STranslationWidgetPicker::IsChecked)
+
+				.Padding(FAppStyle::Get().GetMargin(CheckboxStyle))
+		];
+}
+
+ECheckBoxState STranslationWidgetPicker::IsChecked() const
+{
+	return TranslationPickerManager::IsPickerWindowOpen() ? ECheckBoxState::Checked : ECheckBoxState::Unchecked;
+}
+
+void STranslationWidgetPicker::OnCheckStateChanged(const ECheckBoxState NewCheckedState)
+{
+	if (TranslationPickerManager::IsPickerWindowOpen())
+	{
+		TranslationPickerManager::ClosePickerWindow();
+	}
+	else
+	{
+		TranslationPickerManager::OpenPickerWindow();
+	}
+}
+
+TSharedPtr<SWindow> TranslationPickerManager::PickerWindow = TSharedPtr<SWindow>();
+TSharedPtr<STranslationPickerFloatingWindow> TranslationPickerManager::PickerWindowWidget = TSharedPtr<STranslationPickerFloatingWindow>();
+
+bool TranslationPickerManager::OpenPickerWindow()
+{
+	// Not picking previously, launch a picker window
+	if (!PickerWindow.IsValid())
+	{
+		TSharedRef<SWindow> NewWindow = SWindow::MakeCursorDecorator();
+		NewWindow->SetSizingRule(ESizingRule::FixedSize);
+		// The Edit window and Floating window should be roughly the same size, so it isn't too distracting switching between them
+		NewWindow->Resize(FVector2D(STranslationPickerEditWindow::DefaultEditWindowWidth, STranslationPickerEditWindow::DefaultEditWindowHeight));
+		NewWindow->MoveWindowTo(FSlateApplication::Get().GetCursorPos());
+		PickerWindow = NewWindow;
+
+		NewWindow->SetContent(
+			SAssignNew(PickerWindowWidget, STranslationPickerFloatingWindow)
+			.ParentWindow(NewWindow)
+			);
+
+		TSharedPtr<SWindow> RootWindow = FGlobalTabmanager::Get()->GetRootWindow();
+		if (RootWindow.IsValid())
+		{
+			FSlateApplication::Get().AddWindowAsNativeChild(NewWindow, RootWindow.ToSharedRef());
+		}
+		else
+		{
+			FSlateApplication::Get().AddWindow(NewWindow);
+		}
+
+		return true;
+	}
+
+	return false;
+}
+
+void TranslationPickerManager::ClosePickerWindow()
+{
+	if (FSlateApplication::IsInitialized())
+	{
+		FSlateApplication::Get().RequestDestroyWindow(PickerWindow.ToSharedRef());
+	}
+	PickerWindow.Reset(); 
+	PickerWindowWidget.Reset();
+}
+
+#undef LOCTEXT_NAMESPACE
diff --git a/Engine/Source/Developer/TranslationEditor/Private/TranslationPickerWidget.h b/Engine/Source/Developer/TranslationEditor/Private/TranslationPickerWidget.h
new file mode 100644
--- /dev/null
+++ b/Engine/Source/Developer/TranslationEditor/Private/TranslationPickerWidget.h
@@ -0,0 +1,54 @@
+// Copyright Epic Games, Inc. All Rights Reserved.
+#pragma once
+
+#include "HAL/Platform.h"
+#include "Templates/SharedPointer.h"
+#include "Widgets/DeclarativeSyntaxSupport.h"
+#include "Widgets/SCompoundWidget.h"
+
+#define LOCTEXT_NAMESPACE "TranslationPicker"
+
+class STranslationPickerFloatingWindow;
+class SWindow;
+
+enum class ECheckBoxState : uint8;
+
+class TranslationPickerManager
+{
+public:
+	static TSharedPtr<SWindow> PickerWindow;
+	static TSharedPtr<STranslationPickerFloatingWindow> PickerWindowWidget;
+
+	static bool IsPickerWindowOpen() { return PickerWindow.IsValid(); }
+
+	static bool OpenPickerWindow();
+
+	static void ClosePickerWindow();
+};
+
+/** Widget used to launch a 'picking' session */
+class STranslationWidgetPicker : public SCompoundWidget
+{
+public:
+	SLATE_BEGIN_ARGS(STranslationWidgetPicker) {}
+
+	SLATE_END_ARGS()
+
+	void Construct(const FArguments& InArgs);
+
+private:
+
+	/**
+	* Called by slate to determine if this button should appear checked
+	*
+	* @return ECheckBoxState::Checked if it should be checked, ECheckBoxState::Unchecked if not.
+	*/
+	ECheckBoxState IsChecked() const;
+
+	/**
+	* Called by Slate when this tool bar check box button is toggled
+	*/
+	void OnCheckStateChanged(const ECheckBoxState NewCheckedState);
+};
+
+#undef LOCTEXT_NAMESPACE
diff --git a/Engine/Source/Developer/TranslationEditor/Private/TranslationUnit.cpp b/Engine/Source/Developer/TranslationEditor/Private/TranslationUnit.cpp
new file mode 100644
--- /dev/null
+++ b/Engine/Source/Developer/TranslationEditor/Private/TranslationUnit.cpp
@@ -0,0 +1,30 @@
+// Copyright Epic Games, Inc. All Rights Reserved.
+#include "TranslationUnit.h"
+
+#include "Misc/AssertionMacros.h"
+#include "UObject/NameTypes.h"
+#include "UObject/UnrealNames.h"
+#include "UObject/UnrealType.h"
+
+UTranslationUnit::UTranslationUnit( const FObjectInitializer& ObjectInitializer )
+	: Super(ObjectInitializer)
+{
+
+}
+
+#if WITH_EDITOR
+void UTranslationUnit::PostEditChangeProperty(struct FPropertyChangedEvent& PropertyChangedEvent)
+{
+	Super::PostEditChangeProperty(PropertyChangedEvent);
+
+	const FName Name = (PropertyChangedEvent.Property != nullptr) ? PropertyChangedEvent.Property->GetFName() : NAME_None;
+
+	if (Name == GET_MEMBER_NAME_CHECKED(UTranslationUnit, Translation))
+	{
+		// Consider modifying the translation to be an implicit review
+		HasBeenReviewed = true;
+	}
+
+	TranslationUnitPropertyChangedEvent.Broadcast(Name);
+}
+#endif // WITH_EDITOR
diff --git a/Engine/Source/Developer/TranslationEditor/Private/TranslationUnit.h b/Engine/Source/Developer/TranslationEditor/Private/TranslationUnit.h
new file mode 100644
--- /dev/null
+++ b/Engine/Source/Developer/TranslationEditor/Private/TranslationUnit.h
@@ -0,0 +1,124 @@
+// Copyright Epic Games, Inc. All Rights Reserved.
+#pragma once
+
+#include "Containers/Array.h"
+#include "Containers/UnrealString.h"
+#include "Delegates/Delegate.h"
+#include "Misc/DateTime.h"
+#include "Templates/SharedPointer.h"
+#include "UObject/Object.h"
+#include "UObject/UObjectGlobals.h"
+
+#include "TranslationUnit.generated.h"
+
+class FLocMetadataObject;
+class FName;
+
+USTRUCT()
+struct FTranslationChange
+{
+	GENERATED_USTRUCT_BODY()
+
+public:
+
+	/** The changelist of this change */
+	UPROPERTY(Category=Translation, VisibleAnywhere)
+	FString Version;
+
+	/** Date of this change */
+	UPROPERTY(Category=Translation, VisibleAnywhere)//, meta=(DisplayName = "Date & Time"))
+	FDateTime DateAndTime = FDateTime(0);
+
+	/** Source at time of this change */
+	UPROPERTY(Category=Translation, VisibleAnywhere, meta=(MultiLine=true))
+	FString Source;
+
+	/** Translation at time of this change */
+	UPROPERTY(Category=Translation, VisibleAnywhere, meta=(MultiLine=true))
+	FString Translation;
+};
+
+USTRUCT()
+struct FTranslationContextInfo
+{
+	GENERATED_USTRUCT_BODY()
+
+public:
+
+	/** The key specified in LOCTEXT */
+	UPROPERTY(Category=Translation, VisibleAnywhere)
+	FString Key;
+
+	/** What file and line this translation is from */
+	UPROPERTY(Category=Translation, VisibleAnywhere)
+	FString Context;
+
+	/** List of previous versions of the source text for this context */
+	UPROPERTY(Category=Translation, VisibleAnywhere)
+	TArray<FTranslationChange> Changes;
+};
+
+UCLASS(hidecategories = Object, MinimalAPI)
+class UTranslationUnit : public UObject
+{
+	GENERATED_UCLASS_BODY()
+
+public:
+
+	/** The localization namespace for this translation */
+	UPROPERTY(Category=Translation, VisibleAnywhere)
+	FString Namespace;
+
+	/** The localization key for this translation */
+	UPROPERTY(Category=Translation, VisibleAnywhere)
+	FString Key;
+
+	/** Original text from the source language */
+	UPROPERTY(Category=Translation, VisibleAnywhere, meta=(MultiLine=true))
+	FString Source;
+
+	/** Translations */
+	UPROPERTY(Category=Translation, EditAnywhere, meta=(MultiLine=true))
+	FString Translation;
+
+	/** Contexts the source was found in */
+	UPROPERTY(Category=Translation, VisibleAnywhere)
+	TArray<FTranslationContextInfo> Contexts;
+
+	/** Whether the changes have been reviewed */
+	UPROPERTY(Category=Translation, EditAnywhere)
+	bool HasBeenReviewed;
+
+	/** If this Translation Unit had a different translation before import, it will be stored here */
+	UPROPERTY(Category=Translation, EditAnywhere, meta=(MultiLine=true))
+	FString TranslationBeforeImport;
+
+	/** Optional: Which Locres File this translation is in */
+	UPROPERTY(Category = Translation, EditAnywhere)
+	FString LocresPath;
+
+	/**
+	* Returns an event delegate that is executed when a property has changed.
+	*
+	* @return The delegate.
+	*/
+	DECLARE_EVENT_OneParam(UTranslationUnit, FTranslationUnitPropertyChangedEvent, FName /*PropertyName*/);
+	FTranslationUnitPropertyChangedEvent& OnPropertyChanged() { return TranslationUnitPropertyChangedEvent; }
+
+	TSharedPtr<FLocMetadataObject> KeyMetaDataObject;
+
+#if WITH_EDITOR
+protected:
+	/**
+	* Called when a property on this object has been modified externally
+	*
+	* @param PropertyThatChanged the property that was modified
+	*/
+	virtual void PostEditChangeProperty(struct FPropertyChangedEvent& PropertyChangedEvent) override;
+#endif // WITH_EDITOR
+
+private:
+
+	// Holds an event delegate that is executed when a property has changed.
+	FTranslationUnitPropertyChangedEvent TranslationUnitPropertyChangedEvent;
+};
diff --git a/Engine/Source/Developer/TranslationEditor/Public/ITranslationEditor.h b/Engine/Source/Developer/TranslationEditor/Public/ITranslationEditor.h
new file mode 100644
--- /dev/null
+++ b/Engine/Source/Developer/TranslationEditor/Public/ITranslationEditor.h
@@ -0,0 +1,52 @@
+// Copyright Epic Games, Inc. All Rights Reserved.
+#pragma once
+#if WITH_EDITOR
+
+#include "CoreMinimal.h"
+#include "Toolkits/AssetEditorToolkit.h"
+#include "LocalizationTargetTypes.h"
+
+/** Translation Editor public interface */
+class ITranslationEditor : public FAssetEditorToolkit
+#else
+class ITranslationEditor
+#endif // WITH_EDITOR
+{
+public:
+	TRANSLATIONEDITOR_API static void OpenTranslationPicker();
+
+#if WITH_EDITOR
+	TRANSLATIONEDITOR_API static void OpenTranslationEditor(const FString& InManifestFile, const FString& InNativeArchiveFile, const FString& InArchiveFileToEdit);
+	TRANSLATIONEDITOR_API static void OpenTranslationEditor(ULocalizationTarget* const LocalizationTarget, const FString& CultureToEdit);
+
+	ITranslationEditor(const FString& InManifestFile, const FString& InArchiveFile, ULocalizationTarget* const InAssociatedLocalizationTarget)
+		: ManifestFilePath(InManifestFile)
+		, ArchiveFilePath(InArchiveFile)
+		, AssociatedLocalizationTarget(InAssociatedLocalizationTarget)
+	{}
+
+	virtual void OnClose() override;
+
+protected:
+
+	/** The path to the manifest file being used for contexts. */
+	FString ManifestFilePath;
+	/** The path to the archive file being edited. */
+	FString ArchiveFilePath;
+	/** The localization target associated with the files being used/edited, if any. */
+	TWeakObjectPtr<ULocalizationTarget> AssociatedLocalizationTarget;
+
+private:
+
+	/** To keep track of what translation editors are open editing which manifest files */
+	static TMap<FString, ITranslationEditor*> OpenTranslationEditors;
+	
+	/** Called on open to add this to our list of open translation editors */
+	void RegisterTranslationEditor();
+	
+	/** Called on close to remove this from our list of open translation editors */
+	void UnregisterTranslationEditor();
+
+#endif // WITH_EDITOR
+};
+
diff --git a/Engine/Source/Developer/TranslationEditor/Public/TranslationEditorMenu.h b/Engine/Source/Developer/TranslationEditor/Public/TranslationEditorMenu.h
new file mode 100644
--- /dev/null
+++ b/Engine/Source/Developer/TranslationEditor/Public/TranslationEditorMenu.h
@@ -0,0 +1,66 @@
+// Copyright Epic Games, Inc. All Rights Reserved.
+#pragma once
+#if WITH_EDITOR
+
+#include "CoreMinimal.h"
+#include "Styling/AppStyle.h"
+#include "Framework/Commands/Commands.h"
+
+class FExtender;
+class FMenuBuilder;
+class FTranslationEditor;
+
+/**
+ * Translation Editor menu
+ */
+class TRANSLATIONEDITOR_API FTranslationEditorMenu
+{
+public:
+	static void SetupTranslationEditorMenu( TSharedPtr< FExtender > Extender, FTranslationEditor& TranslationEditor);
+	static void SetupTranslationEditorToolbar( TSharedPtr< FExtender > Extender, FTranslationEditor& TranslationEditor);
+	
+protected:
+
+	static void FillTranslationMenu( FMenuBuilder& MenuBuilder/*, FTranslationEditor& TranslationEditor*/ );
+};
+
+
+class FTranslationEditorCommands : public TCommands<FTranslationEditorCommands>
+{
+public:
+	/** Constructor */
+	FTranslationEditorCommands() 
+		: TCommands<FTranslationEditorCommands>("TranslationEditor", NSLOCTEXT("Contexts", "TranslationEditor", "Translation Editor"), NAME_None, FAppStyle::GetAppStyleSetName())
+	{
+	}
+
+	/** Switch fonts */
+	TSharedPtr<FUICommandInfo> ChangeSourceFont;
+	TSharedPtr<FUICommandInfo> ChangeTranslationTargetFont;
+
+	/** Save the translations to file */
+	TSharedPtr<FUICommandInfo> SaveTranslations;
+
+	/** Save the translations to file */
+	TSharedPtr<FUICommandInfo> PreviewAllTranslationsInEditor;
+	
+	/** Download and import the latest translations from the active Translation Service */
+	TSharedPtr<FUICommandInfo> ImportLatestFromLocalizationService;
+
+	/** Export to PortableObject format (.po) */
+	TSharedPtr<FUICommandInfo> ExportToPortableObjectFormat;
+
+	/** Import from PortableObject format (.po) */
+	TSharedPtr<FUICommandInfo> ImportFromPortableObjectFormat;
+
+	/** Open the tab for searching */
+	TSharedPtr<FUICommandInfo> OpenSearchTab;
+
+	/** Open the translation picker */
+	TSharedPtr<FUICommandInfo> OpenTranslationPicker;
+
+	/** Initialize commands */
+	virtual void RegisterCommands() override;
+};
+
+#endif // WITH_EDITOR
diff --git a/Engine/Source/Developer/TranslationEditor/Public/TranslationEditorModule.h b/Engine/Source/Developer/TranslationEditor/Public/TranslationEditorModule.h
new file mode 100644
--- /dev/null
+++ b/Engine/Source/Developer/TranslationEditor/Public/TranslationEditorModule.h
@@ -0,0 +1,68 @@
+// Copyright Epic Games, Inc. All Rights Reserved.
+
+#pragma once
+
+#include "CoreMinimal.h"
+#include "Modules/ModuleInterface.h"
+#if WITH_EDITOR
+#include "Toolkits/AssetEditorToolkit.h"
+#endif // WITH_EDITOR
+
+class FTranslationEditor;
+class ULocalizationTarget;
+
+class FTranslationEditorModule
+	: public IModuleInterface
+#if WITH_EDITOR
+	, public IHasMenuExtensibility
+#endif // WITH_EDITOR
+{
+
+public:
+	// IModuleInterface
+	virtual void StartupModule() override;
+	virtual void ShutdownModule() override;
+
+	/**
+	 * Open the translation picker.
+	 * @note Alias for ITranslationEditor::OpenTranslationPicker.
+	 */
+	virtual void OpenTranslationPicker();
+
+#if WITH_EDITOR
+	/**
+	 * Creates an instance of translation editor object.  Only virtual so that it can be called across the DLL boundary.
+	 *
+	 * @param ManifestFile			The path to the manifest file to be used for contexts.
+	 * @param NativeArchiveFile		The path to the archive file for the native language.
+	 * @param ArchiveFileToEdit		The path to the archive file to be viewed and edited.
+	 * @param OutLoadedSuccessfully	Whether or not the translation editor was able to load successfully from the .manifest and .archive files
+	 * 
+	 * @return	The new instance of the translation editor
+	 */
+	virtual TSharedRef<FTranslationEditor> CreateTranslationEditor(const FString& ManifestFile, const FString& NativeArchiveFile, const FString& ArchiveFileToEdit, bool& OutLoadedSuccessfully);
+
+	/**
+	 * Creates an instance of translation editor object.  Only virtual so that it can be called across the DLL boundary.
+	 *
+	 * @param LocalizationTarget	The localization target whose data is to be used and edited.
+	 * @param CultureToEdit			The name of a supported culture of the localization target, whose archives should be edited.
+	 * @param OutLoadedSuccessfully	Whether or not the translation editor was able to load successfully from the .manifest and .archive files
+	 * 
+	 * @return	The new instance of the translation editor
+	 */
+	virtual TSharedRef<FTranslationEditor> CreateTranslationEditor(ULocalizationTarget* const LocalizationTarget, const FString& CultureToEdit, bool& OutLoadedSuccessfully);
+
+
+	/** Gets the extensibility managers for outside entities to extend translation editor's menus and toolbars */
+	virtual TSharedPtr<FExtensibilityManager> GetMenuExtensibilityManager() override { return MenuExtensibilityManager; }
+	virtual TSharedPtr<FExtensibilityManager> GetToolbarExtensibilityManager() { return ToolbarExtensibilityManager; }
+
+	/** Translation Editor app identifier string */
+	static const FName TranslationEditorAppIdentifier;
+
+private:
+	TSharedPtr<FExtensibilityManager> MenuExtensibilityManager;
+	TSharedPtr<FExtensibilityManager> ToolbarExtensibilityManager;
+#endif // WITH_EDITOR
+};
diff --git a/Engine/Source/Developer/TranslationEditor/TranslationEditor.Build.cs b/Engine/Source/Developer/TranslationEditor/TranslationEditor.Build.cs
new file mode 100644
--- /dev/null
+++ b/Engine/Source/Developer/TranslationEditor/TranslationEditor.Build.cs
@@ -0,0 +1,65 @@
+// Copyright Epic Games, Inc. All Rights Reserved.
+
+using UnrealBuildTool;
+
+public class TranslationEditor : ModuleRules
+{
+	public TranslationEditor(ReadOnlyTargetRules Target) : base(Target)
+	{
+		PrivateIncludePathModuleNames.Add("LocalizationService");
+
+		PrivateIncludePathModuleNames.AddRange(
+			new string[] {
+				"DesktopPlatform",
+			}
+		);
+
+		PrivateDependencyModuleNames.AddRange(
+			new string[] {
+				"ApplicationCore",
+				"EngineSettings",
+				"InputCore",
+				"Json",
+				"Slate",
+				"SlateCore",
+				"SourceControl",
+				"MessageLog",
+				"LocalizationService",
+				"ToolMenus",
+			}
+		);
+
+		if (Target.bBuildEditor)
+		{
+			PublicIncludePathModuleNames.Add("LevelEditor");
+			PublicIncludePathModuleNames.Add("WorkspaceMenuStructure");
+
+			PrivateDependencyModuleNames.AddRange(
+				new string[] {
+					"PropertyEditor",
+					"EditorFramework",
+					"UnrealEd",
+					"GraphEditor",
+					"Documentation",
+					"LocalizationCommandletExecution",
+				}
+			);
+
+			DynamicallyLoadedModuleNames.AddRange(
+				new string[] {
+					"WorkspaceMenuStructure",
+					"DesktopPlatform",
+				}
+			);
+		}
+
+		PublicDependencyModuleNames.AddRange(
+			new string[] {
+				"Core",
+				"CoreUObject",
+				"Engine",
+				"Localization",
+			}
+		);
+	}
+}
diff --git a/Engine/Source/Editor/TranslationEditor/Private/CustomFontColumn.cpp b/Engine/Source/Editor/TranslationEditor/Private/CustomFontColumn.cpp
deleted file mode 100644
diff --git a/Engine/Source/Editor/TranslationEditor/Private/CustomFontColumn.h b/Engine/Source/Editor/TranslationEditor/Private/CustomFontColumn.h
deleted file mode 100644
diff --git a/Engine/Source/Editor/TranslationEditor/Private/ITranslationEditor.cpp b/Engine/Source/Editor/TranslationEditor/Private/ITranslationEditor.cpp
deleted file mode 100644
diff --git a/Engine/Source/Editor/TranslationEditor/Private/InternationalizationExportSettings.cpp b/Engine/Source/Editor/TranslationEditor/Private/InternationalizationExportSettings.cpp
deleted file mode 100644
diff --git a/Engine/Source/Editor/TranslationEditor/Private/InternationalizationExportSettings.h b/Engine/Source/Editor/TranslationEditor/Private/InternationalizationExportSettings.h
deleted file mode 100644
diff --git a/Engine/Source/Editor/TranslationEditor/Private/TranslationDataManager.cpp b/Engine/Source/Editor/TranslationEditor/Private/TranslationDataManager.cpp
deleted file mode 100644
diff --git a/Engine/Source/Editor/TranslationEditor/Private/TranslationDataManager.h b/Engine/Source/Editor/TranslationEditor/Private/TranslationDataManager.h
deleted file mode 100644
diff --git a/Engine/Source/Editor/TranslationEditor/Private/TranslationEditor.cpp b/Engine/Source/Editor/TranslationEditor/Private/TranslationEditor.cpp
deleted file mode 100644
diff --git a/Engine/Source/Editor/TranslationEditor/Private/TranslationEditor.h b/Engine/Source/Editor/TranslationEditor/Private/TranslationEditor.h
deleted file mode 100644
diff --git a/Engine/Source/Editor/TranslationEditor/Private/TranslationEditorMenu.cpp b/Engine/Source/Editor/TranslationEditor/Private/TranslationEditorMenu.cpp
deleted file mode 100644
diff --git a/Engine/Source/Editor/TranslationEditor/Private/TranslationEditorModule.cpp b/Engine/Source/Editor/TranslationEditor/Private/TranslationEditorModule.cpp
deleted file mode 100644
diff --git a/Engine/Source/Editor/TranslationEditor/Private/TranslationPickerEditWindow.cpp b/Engine/Source/Editor/TranslationEditor/Private/TranslationPickerEditWindow.cpp
deleted file mode 100644
diff --git a/Engine/Source/Editor/TranslationEditor/Private/TranslationPickerEditWindow.h b/Engine/Source/Editor/TranslationEditor/Private/TranslationPickerEditWindow.h
deleted file mode 100644
diff --git a/Engine/Source/Editor/TranslationEditor/Private/TranslationPickerFloatingWindow.cpp b/Engine/Source/Editor/TranslationEditor/Private/TranslationPickerFloatingWindow.cpp
deleted file mode 100644
diff --git a/Engine/Source/Editor/TranslationEditor/Private/TranslationPickerFloatingWindow.h b/Engine/Source/Editor/TranslationEditor/Private/TranslationPickerFloatingWindow.h
deleted file mode 100644
diff --git a/Engine/Source/Editor/TranslationEditor/Private/TranslationPickerWidget.cpp b/Engine/Source/Editor/TranslationEditor/Private/TranslationPickerWidget.cpp
deleted file mode 100644
diff --git a/Engine/Source/Editor/TranslationEditor/Private/TranslationPickerWidget.h b/Engine/Source/Editor/TranslationEditor/Private/TranslationPickerWidget.h
deleted file mode 100644
diff --git a/Engine/Source/Editor/TranslationEditor/Private/TranslationUnit.cpp b/Engine/Source/Editor/TranslationEditor/Private/TranslationUnit.cpp
deleted file mode 100644
diff --git a/Engine/Source/Editor/TranslationEditor/Private/TranslationUnit.h b/Engine/Source/Editor/TranslationEditor/Private/TranslationUnit.h
deleted file mode 100644
diff --git a/Engine/Source/Editor/TranslationEditor/Public/ITranslationEditor.h b/Engine/Source/Editor/TranslationEditor/Public/ITranslationEditor.h
deleted file mode 100644
diff --git a/Engine/Source/Editor/TranslationEditor/Public/TranslationEditorMenu.h b/Engine/Source/Editor/TranslationEditor/Public/TranslationEditorMenu.h
deleted file mode 100644
diff --git a/Engine/Source/Editor/TranslationEditor/Public/TranslationEditorModule.h b/Engine/Source/Editor/TranslationEditor/Public/TranslationEditorModule.h
deleted file mode 100644
diff --git a/Engine/Source/Editor/TranslationEditor/TranslationEditor.Build.cs b/Engine/Source/Editor/TranslationEditor/TranslationEditor.Build.cs
deleted file mode 100644
diff --git a/Engine/Source/Runtime/Engine/Engine.Build.cs b/Engine/Source/Runtime/Engine/Engine.Build.cs
--- a/Engine/Source/Runtime/Engine/Engine.Build.cs
+++ b/Engine/Source/Runtime/Engine/Engine.Build.cs
@@ -55,12 +55,14 @@ public class Engine : ModuleRules
 			PrivateIncludePathModuleNames.AddRange(
 				new string[] {
 					"SlateReflector",
+					"TranslationEditor",
 				}
 			);
 
 			DynamicallyLoadedModuleNames.AddRange(
 				new string[] {
 					"SlateReflector",
+					"TranslationEditor",
 				}
 			);
 
diff --git a/Engine/Source/Runtime/Engine/Private/EngineGlobals.cpp b/Engine/Source/Runtime/Engine/Private/EngineGlobals.cpp
--- a/Engine/Source/Runtime/Engine/Private/EngineGlobals.cpp
+++ b/Engine/Source/Runtime/Engine/Private/EngineGlobals.cpp
@@ -15,6 +15,7 @@
 
 #if !UE_BUILD_SHIPPING && PLATFORM_DESKTOP 
 #include "ISlateReflectorModule.h"
+#include "TranslationEditorModule.h"
 #endif
 
 class IRendererModule;
@@ -191,7 +192,6 @@ static void ShowFontAtlasVisualizer()
 {
 	static const FName SlateReflectorModuleName("SlateReflector");
 	FModuleManager::LoadModuleChecked<ISlateReflectorModule>(SlateReflectorModuleName).DisplayFontAtlasVisualizer();
-
 }
 
 static FAutoConsoleCommand GShowWidgetReflector
@@ -215,4 +215,15 @@ static FAutoConsoleCommand GShowFontAtlasReflector
 	FConsoleCommandDelegate::CreateStatic(ShowFontAtlasVisualizer)
 );
 
+static FAutoConsoleCommand GTranslationPicker
+(
+	TEXT("TranslationPicker"),
+	TEXT("Displays the translation picker"),
+	FConsoleCommandDelegate::CreateStatic([]()
+	{
+		static const FName TranslationEditorModuleName("TranslationEditor");
+		FModuleManager::LoadModuleChecked<FTranslationEditorModule>(TranslationEditorModuleName).OpenTranslationPicker();
+	})
+);
+
 #endif
