TranslationPicker
- Avoid crash switching from Edit to Picking window. On Enter key input, IInputProcessor would call the manager to delete its owner and itself, instead just call the owner Close() to align with the Picking window.
- Avoid less frequent crash switching windows by switching to RequestDestroyWindow, which puts the window in a queue for release
- On Enter in Picking window, if no texts are selected, do not close the Picking window silently, just leave it open. We now have Escape key to close in both windows.
- Enable in-editor changes to DisplayString in the Edit window, not just game

diff --git a/Engine/Source/Developer/TranslationEditor/Private/TranslationPickerEditWindow.cpp b/Engine/Source/Developer/TranslationEditor/Private/TranslationPickerEditWindow.cpp
--- a/Engine/Source/Developer/TranslationEditor/Private/TranslationPickerEditWindow.cpp
+++ b/Engine/Source/Developer/TranslationEditor/Private/TranslationPickerEditWindow.cpp
@@ -269,10 +269,10 @@ STranslationPickerEditWindow::~STranslationPickerEditWindow()
 
 FReply STranslationPickerEditWindow::Close()
 {
-	if (ParentWindow.IsValid())
+	const TSharedPtr<SWindow> ContainingWindow = FSlateApplication::Get().FindWidgetWindow(AsShared());
+	if (ContainingWindow.IsValid())
 	{
-		FSlateApplication::Get().RequestDestroyWindow(ParentWindow.Pin().ToSharedRef());
-		ParentWindow.Reset();
+		ContainingWindow->RequestDestroyWindow();
 	}
 
 	return FReply::Handled();
@@ -573,8 +573,7 @@ FReply STranslationPickerEditWidget::SaveAndPreview()
 	TArray<UTranslationUnit*> TempArray;
 	TempArray.Add(TranslationUnit);
 	FTranslationDataManager::SaveSelectedTranslations(TempArray, ILocalizationServiceModule::Get().GetProvider().IsEnabled() && TranslationPickerSettings->bSubmitTranslationPickerChangesToLocalizationService);
-
-#else
+#endif // WITH_EDITOR
 
 	FTextLocalizationManager::Get().AddOrUpdateDisplayStringInLiveTable(TranslationUnit->Namespace, TranslationUnit->Key, TranslationUnit->Translation, &TranslationUnit->Source);
 
@@ -582,7 +581,6 @@ FReply STranslationPickerEditWidget::SaveAndPreview()
 	{
 		CObj->AsCommand()->Execute(/*Args=*/TArray<FString>(), /*InWorld=*/nullptr, *GLog);
 	}
-#endif // WITH_EDITOR
 
 	return FReply::Handled();
 }
diff --git a/Engine/Source/Developer/TranslationEditor/Private/TranslationPickerFloatingWindow.cpp b/Engine/Source/Developer/TranslationEditor/Private/TranslationPickerFloatingWindow.cpp
--- a/Engine/Source/Developer/TranslationEditor/Private/TranslationPickerFloatingWindow.cpp
+++ b/Engine/Source/Developer/TranslationEditor/Private/TranslationPickerFloatingWindow.cpp
@@ -79,12 +79,15 @@ public:
 
 		if (Key == EKeys::Escape)
 		{
-			TranslationPickerManager::ClosePickerWindow();
+			Owner->Close();
 			return true;
 		}
 		else if (Key == EKeys::Enter)
 		{
-			Owner->SwitchToEditWindow();
+			if (Owner->SwitchToEditWindow())
+			{
+				Owner->Close();
+			}
 			return true;
 		}
 		else if (InKeyEvent.IsControlDown())
@@ -160,6 +163,18 @@ STranslationPickerFloatingWindow::~STranslationPickerFloatingWindow()
 	}
 }
 
+FReply STranslationPickerFloatingWindow::Close()
+{
+	const TSharedPtr<SWindow> ContainingWindow = FSlateApplication::Get().FindWidgetWindow(AsShared());
+	if (ContainingWindow.IsValid())
+	{
+		ContainingWindow->RequestDestroyWindow();
+	}
+	TranslationPickerManager::ResetPickerWindow();
+
+	return FReply::Handled();
+}
+
 void STranslationPickerFloatingWindow::Tick( const FGeometry& AllottedGeometry, const double InCurrentTime, const float InDeltaTime )
 {
 	FWidgetPath Path = FSlateApplication::Get().LocateWindowUnderMouse(FSlateApplication::Get().GetCursorPos(), FSlateApplication::Get().GetInteractiveTopLevelWindows(), true);
@@ -371,7 +386,7 @@ void STranslationPickerFloatingWindow::PickTextFromChildWidgets(TSharedRef<SWidg
 	}
 }
 
-void STranslationPickerFloatingWindow::SwitchToEditWindow()
+bool STranslationPickerFloatingWindow::SwitchToEditWindow()
 {
 	if (PickedTexts.Num() > 0)
 	{
@@ -401,9 +416,11 @@ void STranslationPickerFloatingWindow::SwitchToEditWindow()
 		}
 
 		NewWindow->MoveWindowTo(ParentWindow.Pin()->GetPositionInScreen());
-	}
 
-	TranslationPickerManager::ClosePickerWindow();
+		return true;
+	}
+	return false;
 }
 
+
 #undef LOCTEXT_NAMESPACE
diff --git a/Engine/Source/Developer/TranslationEditor/Private/TranslationPickerFloatingWindow.h b/Engine/Source/Developer/TranslationEditor/Private/TranslationPickerFloatingWindow.h
--- a/Engine/Source/Developer/TranslationEditor/Private/TranslationPickerFloatingWindow.h
+++ b/Engine/Source/Developer/TranslationEditor/Private/TranslationPickerFloatingWindow.h
@@ -34,6 +34,8 @@ class STranslationPickerFloatingWindow : public SCompoundWidget
 private:
 	friend class FTranslationPickerInputProcessor;
 
+	FReply Close();
+
 	virtual void Tick( const FGeometry& AllottedGeometry, const double InCurrentTime, const float InDeltaTime ) override;
 
 	/** Pull the FText reference out of an SWidget */
@@ -43,7 +45,7 @@ private:
 	void PickTextFromChildWidgets(TSharedRef<SWidget> Widget);
 
 	/** Switch from floating window to edit window */
-	void SwitchToEditWindow();
+	bool SwitchToEditWindow();
 
 	/** Input processor used to capture the 'Esc' key */
 	TSharedPtr<FTranslationPickerInputProcessor> InputProcessor;
diff --git a/Engine/Source/Developer/TranslationEditor/Private/TranslationPickerWidget.cpp b/Engine/Source/Developer/TranslationEditor/Private/TranslationPickerWidget.cpp
--- a/Engine/Source/Developer/TranslationEditor/Private/TranslationPickerWidget.cpp
+++ b/Engine/Source/Developer/TranslationEditor/Private/TranslationPickerWidget.cpp
@@ -129,7 +129,7 @@ TSharedPtr<STranslationPickerFloatingWindow> TranslationPickerManager::PickerWin
 bool TranslationPickerManager::OpenPickerWindow()
 {
 	// Not picking previously, launch a picker window
-	if (!PickerWindow.IsValid())
+	if (!PickerWindow.IsValid() && !PickerWindowWidget.IsValid())
 	{
 		TSharedRef<SWindow> NewWindow = SWindow::MakeCursorDecorator();
 		NewWindow->SetSizingRule(ESizingRule::FixedSize);
@@ -165,7 +165,14 @@ void TranslationPickerManager::ClosePickerWindow()
 	{
 		FSlateApplication::Get().RequestDestroyWindow(PickerWindow.ToSharedRef());
 	}
-	PickerWindow.Reset(); 
+
+	ResetPickerWindow();
+}
+
+void TranslationPickerManager::ResetPickerWindow()
+{
+	PickerWindow.Reset();
+
 	PickerWindowWidget.Reset();
 }
 
diff --git a/Engine/Source/Developer/TranslationEditor/Private/TranslationPickerWidget.h b/Engine/Source/Developer/TranslationEditor/Private/TranslationPickerWidget.h
--- a/Engine/Source/Developer/TranslationEditor/Private/TranslationPickerWidget.h
+++ b/Engine/Source/Developer/TranslationEditor/Private/TranslationPickerWidget.h
@@ -24,6 +24,8 @@ public:
 	static bool OpenPickerWindow();
 
 	static void ClosePickerWindow();
+
+	static void ResetPickerWindow();
 };
 
 /** Widget used to launch a 'picking' session */
