TranslationPicker with ENABLE_LOC_TESTING

#trivial

diff --git a/Engine/Source/Developer/TranslationEditor/Private/TranslationPickerEditWindow.cpp b/Engine/Source/Developer/TranslationEditor/Private/TranslationPickerEditWindow.cpp
--- a/Engine/Source/Developer/TranslationEditor/Private/TranslationPickerEditWindow.cpp
+++ b/Engine/Source/Developer/TranslationEditor/Private/TranslationPickerEditWindow.cpp
@@ -575,12 +575,14 @@ FReply STranslationPickerEditWidget::SaveAndPreview()
 	FTranslationDataManager::SaveSelectedTranslations(TempArray, ILocalizationServiceModule::Get().GetProvider().IsEnabled() && TranslationPickerSettings->bSubmitTranslationPickerChangesToLocalizationService);
 #endif // WITH_EDITOR
 
+#if ENABLE_LOC_TESTING
 	FTextLocalizationManager::Get().AddOrUpdateDisplayStringInLiveTable(TranslationUnit->Namespace, TranslationUnit->Key, TranslationUnit->Translation, &TranslationUnit->Source);
 
 	if (IConsoleObject* CObj = IConsoleManager::Get().FindConsoleObject(TEXT("Slate.TriggerInvalidate")))
 	{
 		CObj->AsCommand()->Execute(/*Args=*/TArray<FString>(), /*InWorld=*/nullptr, *GLog);
 	}
+#endif // ENABLE_LOC_TESTING
 
 	return FReply::Handled();
 }
