Translation Picker - Hotkey to disable 3D view mouse movement

diff --git a/Engine/Source/Developer/TranslationEditor/Private/TranslationPickerFloatingWindow.cpp b/Engine/Source/Developer/TranslationEditor/Private/TranslationPickerFloatingWindow.cpp
--- a/Engine/Source/Developer/TranslationEditor/Private/TranslationPickerFloatingWindow.cpp
+++ b/Engine/Source/Developer/TranslationEditor/Private/TranslationPickerFloatingWindow.cpp
@@ -4,10 +4,14 @@
 
 #include "Containers/SortedMap.h"
 #include "Containers/UnrealString.h"
+#include "CoreGlobals.h"
+#include "Engine/Engine.h"
+#include "Engine/GameEngine.h"
 #include "Framework/Application/IInputProcessor.h"
 #include "Framework/Application/SlateApplication.h"
 #include "Framework/Docking/TabManager.h"
 #include "Framework/Text/TextLayout.h"
+#include "GameFramework/PlayerController.h"
 #include "HAL/Platform.h"
 #include "HAL/PlatformCrt.h"
 #include "HAL/PlatformApplicationMisc.h"
@@ -16,6 +20,7 @@
 #include "Internationalization/Internationalization.h"
 #include "Internationalization/Text.h"
 #include "Internationalization/TextNamespaceUtil.h"
+#include "Kismet/GameplayStatics.h"
 #include "Layout/ArrangedChildren.h"
 #include "Layout/ArrangedWidget.h"
 #include "Layout/BasicLayoutWidgetSlot.h"
@@ -26,6 +31,7 @@
 #include "Math/Vector2D.h"
 #include "Misc/Attribute.h"
 #if WITH_EDITOR
+#include "Editor.h"
 #include "SDocumentationToolTip.h"
 #endif // WITH_EDITOR
 #include "SlotBase.h"
@@ -79,6 +85,7 @@ public:
 
 		if (Key == EKeys::Escape)
 		{
+			Owner->SetViewportMouseIgnoreLook(false);
 			Owner->Close();
 			return true;
 		}
@@ -86,10 +93,23 @@ public:
 		{
 			if (Owner->SwitchToEditWindow())
 			{
+				Owner->SetViewportMouseIgnoreLook(false);
 				Owner->Close();
 			}
 			return true;
 		}
+		else if (Key == EKeys::Backslash)
+		{
+			if (Owner->bMouseLookInputIgnored)
+			{
+				Owner->SetViewportMouseIgnoreLook(false);
+			}
+			else
+			{
+				Owner->SetViewportMouseIgnoreLook(true);
+			}
+			return true;
+		}
 		else if (InKeyEvent.IsControlDown())
 		{
 			const uint32* KeyCode = nullptr;
@@ -422,5 +442,46 @@ bool STranslationPickerFloatingWindow::SwitchToEditWindow()
 	return false;
 }
 
+void STranslationPickerFloatingWindow::SetViewportMouseIgnoreLook(bool bLookIgnore)
+{
+	// Avoid multiple increments/decrements to AController::IgnoreLookInput, which is a uint8
+	if (bMouseLookInputIgnored == bLookIgnore)
+	{
+		return;
+	}
+
+	if (UWorld* World = GetWorld())
+	{
+		if (World->HasBegunPlay())
+		{
+			if (APlayerController* PlayerController = UGameplayStatics::GetPlayerController(World, 0))
+			{
+				PlayerController->SetIgnoreLookInput(bLookIgnore);
+				bMouseLookInputIgnored = bLookIgnore;
+			}
+		}
+	}
+}
+
+UWorld* STranslationPickerFloatingWindow::GetWorld() const
+{
+#if WITH_EDITOR
+	if (GIsEditor && IsValid(GEditor))
+	{
+		FWorldContext* PIEWorldContext = GEditor->GetPIEWorldContext();
+		if (PIEWorldContext)
+		{
+			return PIEWorldContext->World();
+		}
+		return GEditor->GetEditorWorldContext().World();
+	}
+	else
+#endif // WITH_EDITOR
+	if (UGameEngine* GameEngine = Cast<UGameEngine>(GEngine))
+	{
+		return GameEngine->GetGameWorld();
+	}
+	return nullptr;
+}
 
 #undef LOCTEXT_NAMESPACE
diff --git a/Engine/Source/Developer/TranslationEditor/Private/TranslationPickerFloatingWindow.h b/Engine/Source/Developer/TranslationEditor/Private/TranslationPickerFloatingWindow.h
--- a/Engine/Source/Developer/TranslationEditor/Private/TranslationPickerFloatingWindow.h
+++ b/Engine/Source/Developer/TranslationEditor/Private/TranslationPickerFloatingWindow.h
@@ -47,6 +47,12 @@ private:
 	/** Switch from floating window to edit window */
 	bool SwitchToEditWindow();
 
+	/** Toggle 3D viewport mouse turning */
+	void SetViewportMouseIgnoreLook(bool bLookIgnore);
+
+	/** Get world from editor or engine */
+	UWorld* GetWorld() const;
+
 	/** Input processor used to capture the 'Esc' key */
 	TSharedPtr<FTranslationPickerInputProcessor> InputProcessor;
 
@@ -63,6 +69,8 @@ private:
 	* The path widgets we were hovering over last tick
 	*/
 	FWeakWidgetPath LastTickHoveringWidgetPath;
+
+	bool bMouseLookInputIgnored = false;
 };
 
 #undef LOCTEXT_NAMESPACE
