TranslationPicker: search/filter bar
- Add SSearchBox
- Switched to using SListView, for key focus to work with SSearchBox
- Fixed missing border around list elements in Edit window
- Game draws one less border depth than Editor, compensate by drawing an extra SBorder

diff --git a/Engine/Source/Developer/TranslationEditor/Private/TranslationPickerEditWindow.cpp b/Engine/Source/Developer/TranslationEditor/Private/TranslationPickerEditWindow.cpp
--- a/Engine/Source/Developer/TranslationEditor/Private/TranslationPickerEditWindow.cpp
+++ b/Engine/Source/Developer/TranslationEditor/Private/TranslationPickerEditWindow.cpp
@@ -2,6 +2,7 @@
 
 #include "TranslationPickerEditWindow.h"
 
+#include "Brushes/SlateColorBrush.h"
 #include "Framework/Application/IInputProcessor.h"
 #include "Framework/Application/SlateApplication.h"
 #include "HAL/PlatformCrt.h"
@@ -33,6 +34,7 @@
 #include "Widgets/Input/SCheckBox.h"
 #include "Widgets/Input/SEditableTextBox.h"
 #include "Widgets/Input/SMultiLineEditableTextBox.h"
+#include "Widgets/Input/SSearchBox.h"
 #include "Widgets/Layout/SBorder.h"
 #include "Widgets/Layout/SBox.h"
 #include "Widgets/Layout/SGridPanel.h"
@@ -121,27 +123,6 @@ void STranslationPickerEditWindow::Construct(const FArguments& InArgs)
 		TranslationPickerSettings->bSubmitTranslationPickerChangesToLocalizationService = false;
 	}
 
-	// Add a new Translation Picker Edit Widget for each picked text
-	for (FText PickedText : PickedTexts)
-	{
-		TSharedRef<STranslationPickerEditWidget> NewEditWidget = 
-			SNew(STranslationPickerEditWidget)
-			.PickedText(PickedText)
-			.bAllowEditing(true);
-
-		EditWidgets.Add(NewEditWidget);
-
-		TextsBox->AddSlot()
-			.AutoHeight()
-			.Padding(FMargin(5))
-			[
-				SNew(SBorder)
-				[
-					NewEditWidget
-				]
-			];
-	}
-
 	TSharedPtr<SEditableTextBox> TextBox;
 	float DefaultPadding = 0.0f;
 
@@ -153,13 +134,22 @@ void STranslationPickerEditWindow::Construct(const FArguments& InArgs)
 			SNew(SVerticalBox)
 
 			+SVerticalBox::Slot()
+			.AutoHeight()
 			[
-				SNew(SScrollBox)
-				+SScrollBox::Slot()
-				.Padding(FMargin(8, 5, 8, 5))
-				[
-					TextsBox
-				]
+				SAssignNew(FilterBox, SSearchBox)
+					.HintText(LOCTEXT("FilterBox_Hint", "Filter text entries"))
+					.ToolTipText(LOCTEXT("FilterBox_ToolTip", "Type here to filter the list of text entries."))
+					.SelectAllTextWhenFocused(false)
+					.OnTextChanged(this, &STranslationPickerEditWindow::FilterBox_OnTextChanged)
+					.OnTextCommitted(this, &STranslationPickerEditWindow::FilterBox_OnTextCommitted)
+			]
+
+			+SVerticalBox::Slot()
+			.FillHeight(1.0f)		// Stretch the list vertically to fill up the user-resizable space
+			[
+				SAssignNew(TextListView, STextListView)
+					.ListItemsSource(&FilteredItems)
+					.OnGenerateRow(this, &STranslationPickerEditWindow::TextListView_OnGenerateRow)
 			]
 			
 			+SVerticalBox::Slot()
@@ -252,6 +242,8 @@ void STranslationPickerEditWindow::Construct(const FArguments& InArgs)
 
 	InputProcessor = MakeShared<FTranslationPickerEditInputProcessor>(this);
 	FSlateApplication::Get().RegisterInputPreProcessor(InputProcessor, 0);
+
+	UpdateListItems();
 }
 
 STranslationPickerEditWindow::~STranslationPickerEditWindow()
@@ -291,10 +283,10 @@ FReply STranslationPickerEditWindow::SaveAllAndClose()
 {
 	TArray<UTranslationUnit*> TempArray;
 
-	for (TSharedRef<STranslationPickerEditWidget> EditWidget : EditWidgets)
+	for (TSharedPtr<FTranslationPickerTextItem> EditItem : AllItems)
 	{
-		UTranslationUnit* TranslationUnit = EditWidget->GetTranslationUnitWithAnyChanges();
-		if (TranslationUnit != nullptr && EditWidget->CanSave())
+		UTranslationUnit* TranslationUnit = EditItem->GetTranslationUnitWithAnyChanges();
+		if (TranslationUnit != nullptr && EditItem->CanSave())
 		{
 			TempArray.Add(TranslationUnit);
 		}
@@ -312,131 +304,200 @@ FReply STranslationPickerEditWindow::SaveAllAndClose()
 	return FReply::Handled();
 }
 
-void STranslationPickerEditWidget::Construct(const FArguments& InArgs)
+void STranslationPickerEditWindow::UpdateListItems()
 {
-	PickedText = InArgs._PickedText;
-	bAllowEditing = InArgs._bAllowEditing;
+	AllItems.Reset();
+	FilteredItems.Reset();
+
+	// Add a new Translation Picker Edit Widget for each picked text
+	for (const FText& PickedText : PickedTexts)
+	{
+		TSharedPtr<FTranslationPickerTextItem> Item = FTranslationPickerTextItem::BuildTextItem(PickedText, true);
+		
+		AllItems.Add(Item);
+
+		const FString& FilterBy = FilterText.ToString();
+
+		if (!PickedText.IsEmptyOrWhitespace() &&
+			!PickedText.ToString().Contains(FilterBy) &&
+			!PickedText.BuildSourceString().Contains(FilterBy))
+		{
+			continue;
+		}
+
+		FilteredItems.Add(Item);
+	}
+
+	// Update the list view if we have one
+	if (TextListView.IsValid())
+	{
+		TextListView->RequestListRefresh();
+	}
+}
+
+TSharedPtr<FTranslationPickerTextItem> FTranslationPickerTextItem::BuildTextItem(const FText& InText, bool bAllowEditing)
+{
+	TSharedPtr<FTranslationPickerTextItem> Item = MakeShared<FTranslationPickerTextItem>(InText, bAllowEditing);
 
 	// Try and get the localization information for this text
-	FTextId TextId;
-	FString SourceString;
-	FString TranslationString;
 	{
-		if (const FString* SourceStringPtr = FTextInspector::GetSourceString(PickedText))
+		if (const FString* SourceStringPtr = FTextInspector::GetSourceString(InText))
 		{
-			SourceString = *SourceStringPtr;
+			Item->SourceString = *SourceStringPtr;
 		}
-		TranslationString = FTextInspector::GetDisplayString(PickedText);
-		TextId = FTextInspector::GetTextId(PickedText);
+		Item->TranslationString = FTextInspector::GetDisplayString(InText);
+		Item->TextId = FTextInspector::GetTextId(InText);
 	}
 
 	// Try and find the LocRes the active translation came from
 	// We assume the LocRes is named the same as the localization target
 	FString LocResPath;
-	FString LocTargetName;
-	FString LocResCultureName;
 #if WITH_EDITORONLY_DATA
-	if (!TextId.IsEmpty() && FTextLocalizationManager::Get().GetLocResID(TextId.GetNamespace(), TextId.GetKey(), LocResPath))
+	if (!Item->TextId.IsEmpty() && FTextLocalizationManager::Get().GetLocResID(Item->TextId.GetNamespace(), Item->TextId.GetKey(), LocResPath))
 	{
-		LocTargetName = FPaths::GetBaseFilename(LocResPath);
+		Item->LocTargetName = FPaths::GetBaseFilename(LocResPath);
 
 		const FString CultureFilePath = FPaths::GetPath(LocResPath);
-		LocResCultureName = FPaths::GetBaseFilename(CultureFilePath);
+		Item->LocResCultureName = FPaths::GetBaseFilename(CultureFilePath);
 	}
 #endif // WITH_EDITORONLY_DATA
 
 	// Clean the package localization ID from the namespace (to mirror what the text gatherer does when scraping for translation data)
-	FString CleanNamespace = TextNamespaceUtil::StripPackageNamespace(TextId.GetNamespace().GetChars());
+	FString CleanNamespace = TextNamespaceUtil::StripPackageNamespace(Item->TextId.GetNamespace().GetChars());
 
 	// Save the necessary data in UTranslationUnit for later.  This is what we pass to TranslationDataManager to save our edits
-	TranslationUnit = NewObject<UTranslationUnit>();
-	TranslationUnit->Namespace = CleanNamespace;
-	TranslationUnit->Key = TextId.GetKey().GetChars();
-	TranslationUnit->Source = SourceString;
-	TranslationUnit->Translation = TranslationString;
-	TranslationUnit->LocresPath = LocResPath;
+	Item->TranslationUnit = NewObject<UTranslationUnit>();
+	Item->TranslationUnit->Namespace = CleanNamespace;
+	Item->TranslationUnit->Key = Item->TextId.GetKey().GetChars();
+	Item->TranslationUnit->Source = Item->SourceString;
+	Item->TranslationUnit->Translation = Item->TranslationString;
+	Item->TranslationUnit->LocresPath = LocResPath;
 
 #if WITH_EDITOR
 	// Can only save if we have have an identity and are in a known localization target file
-	bHasRequiredLocalizationInfoForSaving = !TextId.IsEmpty() && !LocTargetName.IsEmpty();
+	Item->bHasRequiredLocalizationInfoForSaving = !Item->TextId.IsEmpty() && !Item->LocTargetName.IsEmpty();
+#endif // WITH_EDITOR
 
+	return Item;
+}
+
+void STranslationPickerEditWindow::FilterBox_OnTextChanged(const FText& InText)
+{
+	FilterText = InText;
+
+	UpdateListItems();
+}
+
+void STranslationPickerEditWindow::FilterBox_OnTextCommitted(const FText& InText, ETextCommit::Type CommitInfo)
+{
+	FilterBox_OnTextChanged(InText);
+}
+
+void STranslationPickerEditWidget::Construct(const FArguments& InArgs, const TSharedRef<STableViewBase>& InOwnerTable, TSharedPtr<FTranslationPickerTextItem> InListItem)
+{
+	Item = InListItem;
+
+	STableRow<TSharedPtr<FTranslationPickerTextItem>>::Construct(STableRow<TSharedPtr<FTranslationPickerTextItem>>::FArguments(), InOwnerTable);
+
+	SetBorderImage(FAppStyle::GetBrush("WhiteBrush"));
+	SetBorderBackgroundColor(FLinearColor(FColor(36, 36, 36, 255)));	// EStyleColor::Panel mot available in game
+
+#if WITH_EDITOR
 	const FTextBlockStyle& BoldText = FAppStyle::Get().GetWidgetStyle<FTextBlockStyle>("RichTextBlock.Bold");
 #else
 	// Bold text not available in game, suppress errors on 'RichTextBlock.Bold'
 	FTextBlockStyle BoldText = FTextBlockStyle::GetDefault();
 	BoldText.SetColorAndOpacity(FLinearColor(FColorList::White));
+	FSlateColorBrush* BorderBrush = new FSlateColorBrush(FLinearColor::White);
 #endif // WITH_EDITOR
 
 	TSharedPtr<SGridPanel> GridPanel;
 
 	// Layout all our data
 	ChildSlot
+	.Padding(FMargin(5))
 	[
-		SNew(SHorizontalBox)
-		
-		+SHorizontalBox::Slot()
-		.FillWidth(1)
-		.Padding(FMargin(5))
+#if !WITH_EDITOR
+		// The editor treats this border as the background. ie. an extra depth. Draw it in game only
+		SNew(SBorder)
+		.BorderBackgroundColor(FLinearColor(1, 1, 1, 0.45f))
+		.BorderImage(BorderBrush)
+		.Padding(FMargin(3.0f, 3.0f))
 		[
-			SNew(SVerticalBox)
-			+SVerticalBox::Slot()
+#endif // WITH_EDITOR
+			SNew(SBorder)
+			.Padding(FMargin(5))		// no effect, remove this line		
 			[
-				SAssignNew(GridPanel, SGridPanel)
-				.FillColumn(1,1)
-				
-				+SGridPanel::Slot(0,0)
-				.Padding(FMargin(2.5))
-				.HAlign(HAlign_Right)
-				[
-					SNew(STextBlock)
-					.TextStyle(&BoldText)
-					.Text(LOCTEXT("SourceLabel", "Source:"))
-				]
+				SNew(SHorizontalBox)
 
-				+SGridPanel::Slot(0, 1)
-				.Padding(FMargin(2.5))
-				.HAlign(HAlign_Right)
+				+SHorizontalBox::Slot()
+				.FillWidth(1)
+				.Padding(FMargin(5))
 				[
-					SNew(SBox)
-					// Hide translation if we don't have necessary information to modify
-					.Visibility(!bHasRequiredLocalizationInfoForSaving ? EVisibility::Collapsed : EVisibility::Visible)
+					SNew(SVerticalBox)
+					+SVerticalBox::Slot()
 					[
-						SNew(STextBlock)
-						.TextStyle(&BoldText)
+						SAssignNew(GridPanel, SGridPanel)
+						.FillColumn(1,1)
+				
+						+SGridPanel::Slot(0,0)
+						.Padding(FMargin(2.5))
+						.HAlign(HAlign_Right)
+						[
+							SNew(STextBlock)
+							.TextStyle(&BoldText)
+							.Text(LOCTEXT("SourceLabel", "Source:"))
+						]
+
+						+SGridPanel::Slot(0, 1)
+						.Padding(FMargin(2.5))
+						.HAlign(HAlign_Right)
+						[
+							SNew(SBox)
+							// Hide translation if we don't have necessary information to modify
+							.Visibility(!Item->bHasRequiredLocalizationInfoForSaving ? EVisibility::Collapsed : EVisibility::Visible)
+							[
+								SNew(STextBlock)
+								.TextStyle(&BoldText)
 #if WITH_EDITORONLY_DATA
-						.Text(FText::Format(LOCTEXT("TranslationLabelWithCulture", "Translation ({0}):"), FText::AsCultureInvariant(LocResCultureName)))
+								.Text(FText::Format(LOCTEXT("TranslationLabelWithCulture", "Translation ({0}):"), FText::AsCultureInvariant(Item->LocResCultureName)))
 #else
-						.Text(LOCTEXT("TranslationLabel", "Translation:"))
+								.Text(LOCTEXT("TranslationLabel", "Translation:"))
 #endif
-					]
-				]
+							]
+						]
 				
-				+SGridPanel::Slot(1, 0)
-				.Padding(FMargin(2.5))
-				[
-					SNew(SMultiLineEditableTextBox)
-					.IsReadOnly(true)
-					.Text(FText::AsCultureInvariant(SourceString))
-				]
-
-				+SGridPanel::Slot(1, 1)
-				.Padding(FMargin(2.5))
-				[
-					SNew(SBox)
-					// Hide translation if we don't have necessary information to modify
-					.Visibility(!bHasRequiredLocalizationInfoForSaving ? EVisibility::Collapsed : EVisibility::Visible)
-					[
-						SAssignNew(TextBox, SMultiLineEditableTextBox)
-						.IsReadOnly(!bAllowEditing || !bHasRequiredLocalizationInfoForSaving)
-						.Text(FText::AsCultureInvariant(TranslationString))
-						.HintText(LOCTEXT("TranslationEditTextBox_HintText", "Enter/edit translation here."))
+						+SGridPanel::Slot(1, 0)
+						.Padding(FMargin(2.5))
+						[
+							SNew(SMultiLineEditableTextBox)
+							.IsReadOnly(true)
+							.Text(FText::AsCultureInvariant(Item->SourceString))
+						]
+
+						+SGridPanel::Slot(1, 1)
+						.Padding(FMargin(2.5))
+						[
+							SNew(SBox)
+							// Hide translation if we don't have necessary information to modify
+							.Visibility(!Item->bHasRequiredLocalizationInfoForSaving ? EVisibility::Collapsed : EVisibility::Visible)
+							[
+								SAssignNew(Item->TextBox, SMultiLineEditableTextBox)
+								.IsReadOnly(!Item->bAllowEditing || !Item->bHasRequiredLocalizationInfoForSaving)
+								.Text(FText::AsCultureInvariant(Item->TranslationString))
+								.HintText(LOCTEXT("TranslationEditTextBox_HintText", "Enter/edit translation here."))
+							]
+						]
 					]
 				]
 			]
+#if !WITH_EDITOR
 		]
+#endif // !WITH_EDITOR
 	];
 
-	if (!TextId.IsEmpty())
+
+	if (!Item->TextId.IsEmpty())
 	{
 		GridPanel->AddSlot(0, 2)
 			.Padding(FMargin(2.5))
@@ -451,7 +512,7 @@ void STranslationPickerEditWidget::Construct(const FArguments& InArgs)
 			[
 				SNew(SEditableTextBox)
 				.IsReadOnly(true)
-				.Text(FText::AsCultureInvariant(CleanNamespace))
+				.Text(FText::AsCultureInvariant(Item->CleanNamespace))
 			];
 		GridPanel->AddSlot(0, 3)
 			.Padding(FMargin(2.5))
@@ -466,13 +527,12 @@ void STranslationPickerEditWidget::Construct(const FArguments& InArgs)
 			[
 				SNew(SEditableTextBox)
 				.IsReadOnly(true)
-				.Text(FText::AsCultureInvariant(TextId.GetKey().GetChars()))
+				.Text(FText::AsCultureInvariant(Item->TextId.GetKey().GetChars()))
 			];
 		
 		int32 Row = 4;
-		if (bHasRequiredLocalizationInfoForSaving)
+		if (Item->bHasRequiredLocalizationInfoForSaving)
 		{
-
 #if WITH_EDITOR
 			GridPanel->AddSlot(0, Row)
 				.Padding(FMargin(2.5))
@@ -487,7 +547,7 @@ void STranslationPickerEditWidget::Construct(const FArguments& InArgs)
 				[
 					SNew(SEditableTextBox)
 					.IsReadOnly(true)
-					.Text(FText::AsCultureInvariant(LocTargetName))
+					.Text(FText::AsCultureInvariant(Item->LocTargetName))
 				];
 			++Row;
 #endif // WITH_EDITOR
@@ -500,7 +560,7 @@ void STranslationPickerEditWidget::Construct(const FArguments& InArgs)
 					.VAlign(VAlign_Center)
 					.ContentPadding(FAppStyle::GetMargin("StandardDialog.ContentPadding"))
 					.OnClicked(this, &STranslationPickerEditWidget::CopyNamespaceAndKey)
-					.Visibility(bAllowEditing ? EVisibility::Visible : EVisibility::Collapsed)
+					.Visibility(Item->bAllowEditing ? EVisibility::Visible : EVisibility::Collapsed)
 					.Text(LOCTEXT("CopyNamespaceAndKey", "Copy Namespace,Key"))
 				];
 			GridPanel->AddSlot(1, Row)
@@ -512,9 +572,9 @@ void STranslationPickerEditWidget::Construct(const FArguments& InArgs)
 					.VAlign(VAlign_Center)
 					.ContentPadding(FAppStyle::GetMargin("StandardDialog.ContentPadding"))
 					.OnClicked(this, &STranslationPickerEditWidget::SaveAndPreview)
-					.IsEnabled(bHasRequiredLocalizationInfoForSaving)
-					.Visibility(bAllowEditing ? EVisibility::Visible : EVisibility::Collapsed)
-					.Text(bHasRequiredLocalizationInfoForSaving ? LOCTEXT("SaveAndPreviewButtonText", "Save and Preview") : LOCTEXT("SaveAndPreviewButtonDisabledText", "Cannot Save"))
+					.IsEnabled(Item->bHasRequiredLocalizationInfoForSaving)
+					.Visibility(Item->bAllowEditing ? EVisibility::Visible : EVisibility::Collapsed)
+					.Text(Item->bHasRequiredLocalizationInfoForSaving ? LOCTEXT("SaveAndPreviewButtonText", "Save and Preview") : LOCTEXT("SaveAndPreviewButtonDisabledText", "Cannot Save"))
 				];
 		}
 		else
@@ -532,15 +592,15 @@ void STranslationPickerEditWidget::Construct(const FArguments& InArgs)
 	else
 	{
 		FText TextNotLocalizableReason = LOCTEXT("TextNotLocalizable_Generic", "This text is not localizable.");
-		if (PickedText.IsCultureInvariant())
+		if (Item->PickedText.IsCultureInvariant())
 		{
 			TextNotLocalizableReason = LOCTEXT("TextNotLocalizable_CultureInvariant", "This text is not localizable (culture-invariant).");
 		}
-		else if (PickedText.IsTransient())
+		else if (Item->PickedText.IsTransient())
 		{
 			TextNotLocalizableReason = LOCTEXT("TextNotLocalizable_Transient", "This text is not localizable (transient).");
 		}
-		else if (!PickedText.ShouldGatherForLocalization())
+		else if (!Item->PickedText.ShouldGatherForLocalization())
 		{
 			TextNotLocalizableReason = LOCTEXT("TextNotLocalizable_InvalidForGather", "This text is not localizable (invalid for gather).");
 		}
@@ -556,7 +616,7 @@ void STranslationPickerEditWidget::Construct(const FArguments& InArgs)
 	}
 }
 
-void STranslationPickerEditWidget::AddReferencedObjects(FReferenceCollector& Collector)
+void FTranslationPickerTextItem::AddReferencedObjects(FReferenceCollector& Collector)
 {
 	Collector.AddReferencedObject(TranslationUnit);
 }
@@ -564,19 +624,19 @@ void STranslationPickerEditWidget::AddReferencedObjects(FReferenceCollector& Col
 FReply STranslationPickerEditWidget::SaveAndPreview()
 {
 	// Update translation string from entered text
-	TranslationUnit->Translation = TextBox->GetText().ToString();
+	Item->TranslationUnit->Translation = Item->TextBox->GetText().ToString();
 
 #if WITH_EDITOR
 	UTranslationPickerSettings* TranslationPickerSettings = FTranslationPickerSettingsManager::Get()->GetSettings();
 
 	// Save the data via translation data manager
 	TArray<UTranslationUnit*> TempArray;
-	TempArray.Add(TranslationUnit);
+	TempArray.Add(Item->TranslationUnit);
 	FTranslationDataManager::SaveSelectedTranslations(TempArray, ILocalizationServiceModule::Get().GetProvider().IsEnabled() && TranslationPickerSettings->bSubmitTranslationPickerChangesToLocalizationService);
 #endif // WITH_EDITOR
 
 #if ENABLE_LOC_TESTING
-	FTextLocalizationManager::Get().AddOrUpdateDisplayStringInLiveTable(TranslationUnit->Namespace, TranslationUnit->Key, TranslationUnit->Translation, &TranslationUnit->Source);
+	FTextLocalizationManager::Get().AddOrUpdateDisplayStringInLiveTable(Item->TranslationUnit->Namespace, Item->TranslationUnit->Key, Item->TranslationUnit->Translation, &Item->TranslationUnit->Source);
 
 	if (IConsoleObject* CObj = IConsoleManager::Get().FindConsoleObject(TEXT("Slate.TriggerInvalidate")))
 	{
@@ -589,7 +649,7 @@ FReply STranslationPickerEditWidget::SaveAndPreview()
 
 FReply STranslationPickerEditWidget::CopyNamespaceAndKey()
 {
-	const FString CopyString = FString::Printf(TEXT("%s,%s"), *TranslationUnit->Namespace, *TranslationUnit->Key);
+	const FString CopyString = FString::Printf(TEXT("%s,%s"), *Item->TranslationUnit->Namespace, *Item->TranslationUnit->Key);
 	
 	FPlatformApplicationMisc::ClipboardCopy(*CopyString);
 
@@ -598,7 +658,7 @@ FReply STranslationPickerEditWidget::CopyNamespaceAndKey()
 	return FReply::Handled();
 }
 
-UTranslationUnit* STranslationPickerEditWidget::GetTranslationUnitWithAnyChanges()
+UTranslationUnit* FTranslationPickerTextItem::GetTranslationUnitWithAnyChanges()
 {
 	if (TranslationUnit)
 	{
@@ -611,4 +671,9 @@ UTranslationUnit* STranslationPickerEditWidget::GetTranslationUnitWithAnyChanges
 	return nullptr;
 }
 
+TSharedRef<ITableRow> STranslationPickerEditWindow::TextListView_OnGenerateRow(TSharedPtr<FTranslationPickerTextItem> InItem, const TSharedRef<STableViewBase>& OwnerTable)
+{
+	return SNew(STranslationPickerEditWidget, OwnerTable, InItem);
+}
+
 #undef LOCTEXT_NAMESPACE
diff --git a/Engine/Source/Developer/TranslationEditor/Private/TranslationPickerEditWindow.h b/Engine/Source/Developer/TranslationEditor/Private/TranslationPickerEditWindow.h
--- a/Engine/Source/Developer/TranslationEditor/Private/TranslationPickerEditWindow.h
+++ b/Engine/Source/Developer/TranslationEditor/Private/TranslationPickerEditWindow.h
@@ -13,20 +13,25 @@
 #include "UObject/UObjectGlobals.h"
 #include "Widgets/DeclarativeSyntaxSupport.h"
 #include "Widgets/SCompoundWidget.h"
+#include "Widgets/Views/SListView.h"
+#include "Widgets/Views/STableRow.h"
 
 #include "TranslationPickerEditWindow.generated.h"
 
 struct FGeometry;
 struct FKeyEvent;
-
-#define LOCTEXT_NAMESPACE "TranslationPicker"
-
-class FTranslationPickerEditInputProcessor;
+class ITableRow;
 class SBox;
 class SMultiLineEditableTextBox;
+class SSearchBox;
+class STableViewBase;
 class SWindow;
 class UTranslationUnit;
 
+#define LOCTEXT_NAMESPACE "TranslationPicker"
+
+class FTranslationPickerEditInputProcessor;
+
 UCLASS(config = TranslationPickerSettings)
 class UTranslationPickerSettings : public UObject
 {
@@ -85,23 +90,18 @@ private:
 	UTranslationPickerSettings* TranslationPickerSettingsObject;
 };
 
-/** Translation picker edit Widget to handle the display and editing of a single selected FText */
-class STranslationPickerEditWidget : public SCompoundWidget, public FGCObject
+/** A text item in the item list */
+struct FTranslationPickerTextItem : public FGCObject
 {
-	SLATE_BEGIN_ARGS(STranslationPickerEditWidget) {}
-
-	SLATE_ARGUMENT(FText, PickedText)
+	FTranslationPickerTextItem(const FText& InText, bool bAllowEditing) : PickedText(InText), bAllowEditing(bAllowEditing) {}
 
-	SLATE_ARGUMENT(bool, bAllowEditing)
-
-	SLATE_END_ARGS()
-
-	void Construct(const FArguments& InArgs);
+	/** Create new text item */
+	static TSharedPtr<FTranslationPickerTextItem> BuildTextItem(const FText& InText, bool bAllowEditing);
 
 	virtual void AddReferencedObjects(FReferenceCollector& Collector) override;
 	virtual FString GetReferencerName() const override
 	{
-		return TEXT("STranslationPickerEditWidget");
+		return TEXT("FTranslationPickerTextItem");
 	}
 
 	/** Return the translation unit for this text, with any modifications */
@@ -113,31 +113,41 @@ class STranslationPickerEditWidget : public SCompoundWidget, public FGCObject
 		return bAllowEditing && bHasRequiredLocalizationInfoForSaving;
 	}
 
-private:
-
-	virtual bool SupportsKeyboardFocus() const override
-	{
-		return true;
-	}
+	/** The FText that we are using this widget to translate */
+	FText PickedText;
 
-	FReply SaveAndPreview();
+	/** Whether or not to show the save button*/
+	bool bAllowEditing = true;
 
-	FReply CopyNamespaceAndKey();
+	/** Whether or not we were able to find the necessary info for saving */
+	bool bHasRequiredLocalizationInfoForSaving = true;
 
-	/** The FText that we are using this widget to translate */
-	FText PickedText;
+	FTextId TextId;
+	FString SourceString;
+	FString TranslationString;
+	FString LocTargetName;
+	FString LocResCultureName;
+	FString CleanNamespace;
 
 	/** The translation we're editing represented in a UTranslationUnit object */
 	TObjectPtr<UTranslationUnit> TranslationUnit;
 
 	/** The text box for entering/modifying a translation */
 	TSharedPtr<SMultiLineEditableTextBox> TextBox;
+};
 
-	/** Whether or not to show the save button*/
-	bool bAllowEditing;
+class STranslationPickerEditWidget : public STableRow<TSharedPtr<FTranslationPickerTextItem>>
+{
+public:
 
-	/** Whether or not we were able to find the necessary info for saving */
-	bool bHasRequiredLocalizationInfoForSaving = true;
+	void Construct(const FArguments& InArgs, const TSharedRef<STableViewBase>& InOwnerTable, TSharedPtr<FTranslationPickerTextItem> InListItem);
+
+private:
+	FReply SaveAndPreview();
+
+	FReply CopyNamespaceAndKey();
+
+	TSharedPtr<FTranslationPickerTextItem> Item;
 };
 
 /** Translation picker edit window to allow you to translate selected FTexts in place */
@@ -170,7 +180,24 @@ private:
 		return true;
 	}
 
-	/** Input processor used to capture the 'Esc' key */
+	/** Return to picker floating window */
+	FReply RestorePicker();
+
+	/** Save all translations and close */
+	FReply SaveAllAndClose();
+
+	/** Update text list items */
+	void UpdateListItems();
+
+	/** Filters the widgets when the user changes the search text box */
+	void FilterBox_OnTextChanged(const FText& InText);
+
+	/** Filters the widgets when the user hits enter or clears the search box */
+	void FilterBox_OnTextCommitted(const FText& InText, ETextCommit::Type CommitInfo);
+
+	TSharedRef<ITableRow> TextListView_OnGenerateRow(TSharedPtr<FTranslationPickerTextItem> InItem, const TSharedRef<STableViewBase>& OwnerTable);
+
+	/** Input processor used to capture key and mouse events */
 	TSharedPtr<FTranslationPickerEditInputProcessor> InputProcessor;
 
 	/** Handle to the window that contains this widget */
@@ -182,14 +209,19 @@ private:
 	/** The FTexts that we have found under the cursor */
 	TArray<FText> PickedTexts;
 
-	/** All of our current edit widgets */
-	TArray<TSharedRef<STranslationPickerEditWidget>> EditWidgets;
+	/** Full unfiltered list of items */
+	TArray<TSharedPtr<FTranslationPickerTextItem>> AllItems;
 
-	/** Return to picker floating window */
-	FReply RestorePicker();
+	/** Filtered list of items */
+	TArray<TSharedPtr<FTranslationPickerTextItem>> FilteredItems;
 
-	/** Save all translations and close */
-	FReply SaveAllAndClose();
+	/** List view control */
+	typedef SListView<TSharedPtr<FTranslationPickerTextItem>> STextListView;
+	TSharedPtr<STextListView> TextListView;
+
+	/** Box to filter by text */
+	TSharedPtr<SSearchBox> FilterBox;
+	FText FilterText;
 };
 
 #undef LOCTEXT_NAMESPACE
diff --git a/Engine/Source/Developer/TranslationEditor/Private/TranslationPickerFloatingWindow.cpp b/Engine/Source/Developer/TranslationEditor/Private/TranslationPickerFloatingWindow.cpp
--- a/Engine/Source/Developer/TranslationEditor/Private/TranslationPickerFloatingWindow.cpp
+++ b/Engine/Source/Developer/TranslationEditor/Private/TranslationPickerFloatingWindow.cpp
@@ -235,41 +235,15 @@ void STranslationPickerFloatingWindow::Tick( const FGeometry& AllottedGeometry,
 				}
 			}
 
-			TSharedRef<SVerticalBox> TextsBox = SNew(SVerticalBox);
-
-			// Add a new Translation Picker Edit Widget for each picked text
-			for (FText PickedText : PickedTexts)
-			{
-				TextsBox->AddSlot()
-					.AutoHeight()
-					.Padding(FMargin(5))
-					[
-						SNew(SBorder)
-						[
-							SNew(STranslationPickerEditWidget)
-							.PickedText(PickedText)
-							.bAllowEditing(false)
-						]
-					];
-			}
-
 			WindowContents->SetContentWidget(
 				SNew(SVerticalBox)
-				
+
 				+SVerticalBox::Slot()
-				.Padding(0)
-				.FillHeight(1)
-				.Padding(FMargin(5))
+				.FillHeight(1.0f)		// Stretch the list vertically to fill up the user-resizable space
 				[
-					SNew(SScrollBox)
-					.Orientation(EOrientation::Orient_Vertical)
-					.ScrollBarAlwaysVisible(true)
-					
-					+SScrollBox::Slot()
-					.Padding(FMargin(0))
-					[
-						TextsBox
-					]
+					SAssignNew(TextListView, STextListView)
+						.ListItemsSource(&TextListItems)
+						.OnGenerateRow(this, &STranslationPickerFloatingWindow::TextListView_OnGenerateWidget)
 				]
 
 				+SVerticalBox::Slot()
@@ -282,6 +256,8 @@ void STranslationPickerFloatingWindow::Tick( const FGeometry& AllottedGeometry,
 					.Justification(ETextJustify::Center)
 				]
 			);
+
+			UpdateListItems();
 		}
 	}
 
@@ -484,4 +460,28 @@ UWorld* STranslationPickerFloatingWindow::GetWorld() const
 	return nullptr;
 }
 
+void STranslationPickerFloatingWindow::UpdateListItems()
+{
+	TextListItems.Reset();
+
+	for (const FText& PickedText : PickedTexts)
+	{
+		TSharedPtr<FTranslationPickerTextItem> Item = FTranslationPickerTextItem::BuildTextItem(PickedText, false);
+
+		TextListItems.Add(Item);
+	}
+
+	// Update the list view if we have one
+	if (TextListView.IsValid())
+	{
+		TextListView->RequestListRefresh();
+	}
+}
+
+TSharedRef<ITableRow> STranslationPickerFloatingWindow::TextListView_OnGenerateWidget(TSharedPtr<FTranslationPickerTextItem> InItem, const TSharedRef<STableViewBase>& OwnerTable)
+{
+	return SNew(STranslationPickerEditWidget, OwnerTable, InItem);
+}
+
 #undef LOCTEXT_NAMESPACE
+
diff --git a/Engine/Source/Developer/TranslationEditor/Private/TranslationPickerFloatingWindow.h b/Engine/Source/Developer/TranslationEditor/Private/TranslationPickerFloatingWindow.h
--- a/Engine/Source/Developer/TranslationEditor/Private/TranslationPickerFloatingWindow.h
+++ b/Engine/Source/Developer/TranslationEditor/Private/TranslationPickerFloatingWindow.h
@@ -7,16 +7,20 @@
 #include "Templates/SharedPointer.h"
 #include "Widgets/DeclarativeSyntaxSupport.h"
 #include "Widgets/SCompoundWidget.h"
+#include "Widgets/Views/SListView.h"
 
 class FText;
+class ITableRow;
+class STableViewBase;
+class SToolTip;
 class SWidget;
 class SWindow;
 struct FGeometry;
+struct FTranslationPickerTextItem;
 
 #define LOCTEXT_NAMESPACE "TranslationPicker"
 
 class FTranslationPickerInputProcessor;
-class SToolTip;
 
 /** Translation picker floating window to show details of FText(s) under cursor, and allow in-place translation via TranslationPickerEditWindow */
 class STranslationPickerFloatingWindow : public SCompoundWidget
@@ -47,13 +51,18 @@ private:
 	/** Switch from floating window to edit window */
 	bool SwitchToEditWindow();
 
+	/** Update text list items */
+	void UpdateListItems();
+
 	/** Toggle 3D viewport mouse turning */
 	void SetViewportMouseIgnoreLook(bool bLookIgnore);
 
 	/** Get world from editor or engine */
 	UWorld* GetWorld() const;
 
-	/** Input processor used to capture the 'Esc' key */
+	TSharedRef<ITableRow> TextListView_OnGenerateWidget(TSharedPtr<FTranslationPickerTextItem> InItem, const TSharedRef<STableViewBase>& OwnerTable);
+
+	/** Input processor used to capture key and mouse events */
 	TSharedPtr<FTranslationPickerInputProcessor> InputProcessor;
 
 	/** Handle to the window that contains this widget */
@@ -65,9 +74,14 @@ private:
 	/** The FTexts that we have found under the cursor */
 	TArray<FText> PickedTexts;
 
-	/**
-	* The path widgets we were hovering over last tick
-	*/
+	/** List items for the text list */
+	TArray<TSharedPtr<FTranslationPickerTextItem>> TextListItems;
+
+	/** List of all texts */
+	typedef SListView<TSharedPtr<FTranslationPickerTextItem>> STextListView;
+	TSharedPtr<STextListView> TextListView;
+
+	/** The path widgets we were hovering over last tick */
 	FWeakWidgetPath LastTickHoveringWidgetPath;
 
 	bool bMouseLookInputIgnored = false;
