Fixed non-unity/non-pch compile errors

diff --git a/Engine/Source/Developer/TranslationEditor/Public/ITranslationEditor.h b/Engine/Source/Developer/TranslationEditor/Public/ITranslationEditor.h
--- a/Engine/Source/Developer/TranslationEditor/Public/ITranslationEditor.h
+++ b/Engine/Source/Developer/TranslationEditor/Public/ITranslationEditor.h
@@ -1,7 +1,9 @@
 // Copyright Epic Games, Inc. All Rights Reserved.
 #pragma once
-#if WITH_EDITOR
 
+#include "HAL/Platform.h"
+
+#if WITH_EDITOR
 #include "CoreMinimal.h"
 #include "Toolkits/AssetEditorToolkit.h"
 #include "LocalizationTargetTypes.h"
@@ -49,4 +51,3 @@ private:
 
 #endif // WITH_EDITOR
 };
-
