TranslationPicker: Focus search/filter bar when Edit window opened
- Hide scrollbar from floating window, since it cannot be used while the window follows the mouse
- Focus search/filter bar when Edit window opened, instead of having to select the box
- Construct all the Slate elements for Floating window in Construct, to avoid layout artifacts

diff --git a/Engine/Source/Developer/TranslationEditor/Private/TranslationPickerEditWindow.cpp b/Engine/Source/Developer/TranslationEditor/Private/TranslationPickerEditWindow.cpp
--- a/Engine/Source/Developer/TranslationEditor/Private/TranslationPickerEditWindow.cpp
+++ b/Engine/Source/Developer/TranslationEditor/Private/TranslationPickerEditWindow.cpp
@@ -137,19 +137,19 @@ void STranslationPickerEditWindow::Construct(const FArguments& InArgs)
 			.AutoHeight()
 			[
 				SAssignNew(FilterBox, SSearchBox)
-					.HintText(LOCTEXT("FilterBox_Hint", "Filter text entries"))
-					.ToolTipText(LOCTEXT("FilterBox_ToolTip", "Type here to filter the list of text entries."))
-					.SelectAllTextWhenFocused(false)
-					.OnTextChanged(this, &STranslationPickerEditWindow::FilterBox_OnTextChanged)
-					.OnTextCommitted(this, &STranslationPickerEditWindow::FilterBox_OnTextCommitted)
+				.HintText(LOCTEXT("FilterBox_Hint", "Filter text entries"))
+				.ToolTipText(LOCTEXT("FilterBox_ToolTip", "Type here to filter the list of text entries."))
+				.SelectAllTextWhenFocused(false)
+				.OnTextChanged(this, &STranslationPickerEditWindow::FilterBox_OnTextChanged)
+				.OnTextCommitted(this, &STranslationPickerEditWindow::FilterBox_OnTextCommitted)
 			]
 
 			+SVerticalBox::Slot()
 			.FillHeight(1.0f)		// Stretch the list vertically to fill up the user-resizable space
 			[
 				SAssignNew(TextListView, STextListView)
-					.ListItemsSource(&FilteredItems)
-					.OnGenerateRow(this, &STranslationPickerEditWindow::TextListView_OnGenerateRow)
+				.ListItemsSource(&FilteredItems)
+				.OnGenerateRow(this, &STranslationPickerEditWindow::TextListView_OnGenerateRow)
 			]
 			
 			+SVerticalBox::Slot()
@@ -244,6 +244,8 @@ void STranslationPickerEditWindow::Construct(const FArguments& InArgs)
 	FSlateApplication::Get().RegisterInputPreProcessor(InputProcessor, 0);
 
 	UpdateListItems();
+
+	RegisterActiveTimer(0.f, FWidgetActiveTimerDelegate::CreateSP(this, &STranslationPickerEditWindow::SetFocusPostConstruct));
 }
 
 STranslationPickerEditWindow::~STranslationPickerEditWindow()
@@ -381,6 +383,16 @@ TSharedPtr<FTranslationPickerTextItem> FTranslationPickerTextItem::BuildTextItem
 	return Item;
 }
 
+EActiveTimerReturnType STranslationPickerEditWindow::SetFocusPostConstruct(double InCurrentTime, float InDeltaTime)
+{
+	if (FilterBox.IsValid())
+	{
+		FSlateApplication::Get().SetKeyboardFocus(FilterBox, EFocusCause::SetDirectly);
+	}
+
+	return EActiveTimerReturnType::Stop;
+}
+
 void STranslationPickerEditWindow::FilterBox_OnTextChanged(const FText& InText)
 {
 	FilterText = InText;
@@ -422,11 +434,11 @@ void STranslationPickerEditWidget::Construct(const FArguments& InArgs, const TSh
 		SNew(SBorder)
 		.BorderBackgroundColor(FLinearColor(1, 1, 1, 0.45f))
 		.BorderImage(BorderBrush)
-		.Padding(FMargin(3.0f, 3.0f))
+		.Padding(FMargin(2.0f, 2.0f))
 		[
 #endif // WITH_EDITOR
 			SNew(SBorder)
-			.Padding(FMargin(5))		// no effect, remove this line		
+			.Padding(FMargin(5))
 			[
 				SNew(SHorizontalBox)
 
diff --git a/Engine/Source/Developer/TranslationEditor/Private/TranslationPickerEditWindow.h b/Engine/Source/Developer/TranslationEditor/Private/TranslationPickerEditWindow.h
--- a/Engine/Source/Developer/TranslationEditor/Private/TranslationPickerEditWindow.h
+++ b/Engine/Source/Developer/TranslationEditor/Private/TranslationPickerEditWindow.h
@@ -20,6 +20,7 @@
 
 struct FGeometry;
 struct FKeyEvent;
+class FTranslationPickerEditInputProcessor;
 class ITableRow;
 class SBox;
 class SMultiLineEditableTextBox;
@@ -30,8 +31,6 @@ class UTranslationUnit;
 
 #define LOCTEXT_NAMESPACE "TranslationPicker"
 
-class FTranslationPickerEditInputProcessor;
-
 UCLASS(config = TranslationPickerSettings)
 class UTranslationPickerSettings : public UObject
 {
@@ -189,6 +188,9 @@ private:
 	/** Update text list items */
 	void UpdateListItems();
 
+	/** On open, set the keyboard focus to the filter box */
+	EActiveTimerReturnType SetFocusPostConstruct(double InCurrentTime, float InDeltaTime);
+
 	/** Filters the widgets when the user changes the search text box */
 	void FilterBox_OnTextChanged(const FText& InText);
 
diff --git a/Engine/Source/Developer/TranslationEditor/Private/TranslationPickerFloatingWindow.cpp b/Engine/Source/Developer/TranslationEditor/Private/TranslationPickerFloatingWindow.cpp
--- a/Engine/Source/Developer/TranslationEditor/Private/TranslationPickerFloatingWindow.cpp
+++ b/Engine/Source/Developer/TranslationEditor/Private/TranslationPickerFloatingWindow.cpp
@@ -161,6 +161,31 @@ void STranslationPickerFloatingWindow::Construct(const FArguments& InArgs)
 	ParentWindow = InArgs._ParentWindow;
 	WindowContents = SNew(SToolTip);
 
+	WindowContents->SetContentWidget(
+		SNew(SVerticalBox)
+
+		+SVerticalBox::Slot()
+		.FillHeight(1.0f)		// Stretch the list vertically to fill up the user-resizable space
+		[
+			SAssignNew(TextListView, STextListView)
+			.ListItemsSource(&TextListItems)
+			.OnGenerateRow(this, &STranslationPickerFloatingWindow::TextListView_OnGenerateWidget)
+			.ScrollbarVisibility(EVisibility::Collapsed)
+		]
+
+		+SVerticalBox::Slot()
+		.Padding(0)
+		.AutoHeight()
+		.Padding(FMargin(5))
+		[
+			SNew(STextBlock)
+			.Text(PickedTexts.Num() > 0 ?
+				LOCTEXT("TranslationPickerEnterToEdit", "Press Enter to edit translations") :
+				LOCTEXT("TranslationPickerHoverToViewEditEscToQuit", "Hover over text to view/edit translations, or press Esc to quit"))
+			.Justification(ETextJustify::Center)
+		]
+	);
+
 	ChildSlot
 	[
 		WindowContents.ToSharedRef()
@@ -235,28 +260,6 @@ void STranslationPickerFloatingWindow::Tick( const FGeometry& AllottedGeometry,
 				}
 			}
 
-			WindowContents->SetContentWidget(
-				SNew(SVerticalBox)
-
-				+SVerticalBox::Slot()
-				.FillHeight(1.0f)		// Stretch the list vertically to fill up the user-resizable space
-				[
-					SAssignNew(TextListView, STextListView)
-						.ListItemsSource(&TextListItems)
-						.OnGenerateRow(this, &STranslationPickerFloatingWindow::TextListView_OnGenerateWidget)
-				]
-
-				+SVerticalBox::Slot()
-				.Padding(0)
-				.AutoHeight()
-				.Padding(FMargin(5))
-				[
-					SNew(STextBlock)
-					.Text(PickedTexts.Num() > 0 ? LOCTEXT("TranslationPickerEnterToEdit", "Press Enter to edit translations") : LOCTEXT("TranslationPickerHoverToViewEditEscToQuit", "Hover over text to view/edit translations, or press Esc to quit"))
-					.Justification(ETextJustify::Center)
-				]
-			);
-
 			UpdateListItems();
 		}
 	}
@@ -444,8 +447,7 @@ UWorld* STranslationPickerFloatingWindow::GetWorld() const
 #if WITH_EDITOR
 	if (GIsEditor && IsValid(GEditor))
 	{
-		FWorldContext* PIEWorldContext = GEditor->GetPIEWorldContext();
-		if (PIEWorldContext)
+		if (FWorldContext* PIEWorldContext = GEditor->GetPIEWorldContext())
 		{
 			return PIEWorldContext->World();
 		}
