TranslationPicker: Fix for missing Namespace value

diff --git a/Engine/Source/Developer/TranslationEditor/Private/TranslationPickerEditWindow.cpp b/Engine/Source/Developer/TranslationEditor/Private/TranslationPickerEditWindow.cpp
--- a/Engine/Source/Developer/TranslationEditor/Private/TranslationPickerEditWindow.cpp
+++ b/Engine/Source/Developer/TranslationEditor/Private/TranslationPickerEditWindow.cpp
@@ -365,11 +365,11 @@ TSharedPtr<FTranslationPickerTextItem> FTranslationPickerTextItem::BuildTextItem
 #endif // WITH_EDITORONLY_DATA
 
 	// Clean the package localization ID from the namespace (to mirror what the text gatherer does when scraping for translation data)
-	FString CleanNamespace = TextNamespaceUtil::StripPackageNamespace(Item->TextId.GetNamespace().GetChars());
+	Item->CleanNamespace = TextNamespaceUtil::StripPackageNamespace(Item->TextId.GetNamespace().GetChars());
 
 	// Save the necessary data in UTranslationUnit for later.  This is what we pass to TranslationDataManager to save our edits
 	Item->TranslationUnit = NewObject<UTranslationUnit>();
-	Item->TranslationUnit->Namespace = CleanNamespace;
+	Item->TranslationUnit->Namespace = Item->CleanNamespace;
 	Item->TranslationUnit->Key = Item->TextId.GetKey().GetChars();
 	Item->TranslationUnit->Source = Item->SourceString;
 	Item->TranslationUnit->Translation = Item->TranslationString;
