TranslationPicker: Sort Floating window by nearest widget
- Draw boxes for hovered  texts
- Sort the text list by distance from mouse cursor

diff --git a/Engine/Source/Developer/TranslationEditor/Private/TranslationEditorModule.cpp b/Engine/Source/Developer/TranslationEditor/Private/TranslationEditorModule.cpp
--- a/Engine/Source/Developer/TranslationEditor/Private/TranslationEditorModule.cpp
+++ b/Engine/Source/Developer/TranslationEditor/Private/TranslationEditorModule.cpp
@@ -38,6 +38,7 @@ void FTranslationEditorModule::ShutdownModule()
 #if WITH_EDITOR
 	MenuExtensibilityManager.Reset();
 #endif // WITH_EDITOR
+	TranslationPickerManager::RemoveOverlay();
 	TranslationPickerManager::ClosePickerWindow();
 
 #if WITH_UNREAL_DEVELOPER_TOOLS
diff --git a/Engine/Source/Developer/TranslationEditor/Private/TranslationPickerEditWindow.cpp b/Engine/Source/Developer/TranslationEditor/Private/TranslationPickerEditWindow.cpp
--- a/Engine/Source/Developer/TranslationEditor/Private/TranslationPickerEditWindow.cpp
+++ b/Engine/Source/Developer/TranslationEditor/Private/TranslationPickerEditWindow.cpp
@@ -5,7 +5,6 @@
 #include "Brushes/SlateColorBrush.h"
 #include "Framework/Application/IInputProcessor.h"
 #include "Framework/Application/SlateApplication.h"
-#include "HAL/PlatformCrt.h"
 #include "HAL/PlatformApplicationMisc.h"
 #include "ILocalizationServiceModule.h"
 #include "ILocalizationServiceProvider.h"
@@ -19,7 +18,6 @@
 #include "Layout/Children.h"
 #include "Layout/Margin.h"
 #include "Layout/Visibility.h"
-#include "Misc/Attribute.h"
 #include "Misc/CommandLine.h"
 #include "Misc/Parse.h"
 #include "Misc/Paths.h"
@@ -38,7 +36,6 @@
 #include "Widgets/Layout/SBorder.h"
 #include "Widgets/Layout/SBox.h"
 #include "Widgets/Layout/SGridPanel.h"
-#include "Widgets/Layout/SScrollBox.h"
 #include "Widgets/Layout/SUniformGridPanel.h"
 #include "Widgets/SBoxPanel.h"
 #include "Widgets/SWindow.h"
@@ -84,7 +81,7 @@ public:
 
 		if (Key == EKeys::Escape)
 		{
-			Owner->Close();
+			Owner->Exit();
 			return true;
 		}
 		else if (Key == EKeys::Enter)
@@ -92,6 +89,11 @@ public:
 			Owner->RestorePicker();
 			return true;
 		}
+		else if (Key == EKeys::BackSpace)
+		{
+			TranslationPickerManager::bDrawBoxes = !TranslationPickerManager::bDrawBoxes;
+			return true;
+		}
 
 		return false;
 	}
@@ -110,7 +112,6 @@ UTranslationPickerSettings::UTranslationPickerSettings(const FObjectInitializer&
 void STranslationPickerEditWindow::Construct(const FArguments& InArgs)
 {
 	ParentWindow = InArgs._ParentWindow;
-	PickedTexts = InArgs._PickedTexts;
 	WindowContents = SNew(SBox);
 	TSharedRef<SVerticalBox> TextsBox = SNew(SVerticalBox);
 	UTranslationPickerSettings* TranslationPickerSettings = FTranslationPickerSettingsManager::Get()->GetSettings();
@@ -212,7 +213,7 @@ void STranslationPickerEditWindow::Construct(const FArguments& InArgs)
 						.HAlign(HAlign_Center)
 						.VAlign(VAlign_Center)
 						.ContentPadding(FAppStyle::GetMargin("StandardDialog.ContentPadding"))
-						.OnClicked(this, &STranslationPickerEditWindow::SaveAllAndClose)
+						.OnClicked(this, &STranslationPickerEditWindow::SaveAllAndExit)
 						.Text(LOCTEXT("SaveAllAndClose", "Save All and Close"))
 #if WITH_EDITOR
 						.Visibility(EVisibility::Visible)
@@ -227,7 +228,7 @@ void STranslationPickerEditWindow::Construct(const FArguments& InArgs)
 						.HAlign(HAlign_Center)
 						.VAlign(VAlign_Center)
 						.ContentPadding(FAppStyle::GetMargin("StandardDialog.ContentPadding"))
-						.OnClicked(this, &STranslationPickerEditWindow::Close)
+						.OnClicked(this, &STranslationPickerEditWindow::Exit)
 						.Text(LOCTEXT("CancelButton", "Cancel"))
 					]
 				]
@@ -272,16 +273,23 @@ FReply STranslationPickerEditWindow::Close()
 	return FReply::Handled();
 }
 
-FReply STranslationPickerEditWindow::RestorePicker()
+FReply STranslationPickerEditWindow::Exit()
 {
+	TranslationPickerManager::RemoveOverlay();
 	Close();
 
+	return FReply::Handled();
+}
+
+FReply STranslationPickerEditWindow::RestorePicker()
+{
+	Close();
 	TranslationPickerManager::OpenPickerWindow();
 
 	return FReply::Handled();
 }
 
-FReply STranslationPickerEditWindow::SaveAllAndClose()
+FReply STranslationPickerEditWindow::SaveAllAndExit()
 {
 	TArray<UTranslationUnit*> TempArray;
 
@@ -301,7 +309,7 @@ FReply STranslationPickerEditWindow::SaveAllAndClose()
 		FTranslationDataManager::SaveSelectedTranslations(TempArray, ILocalizationServiceModule::Get().GetProvider().IsEnabled() && TranslationPickerSettings->bSubmitTranslationPickerChangesToLocalizationService);
 	}
 
-	Close();
+	Exit();
 
 	return FReply::Handled();
 }
@@ -312,17 +320,17 @@ void STranslationPickerEditWindow::UpdateListItems()
 	FilteredItems.Reset();
 
 	// Add a new Translation Picker Edit Widget for each picked text
-	for (const FText& PickedText : PickedTexts)
+	for (const FTranslationPickerTextAndGeom& PickedText : TranslationPickerManager::PickedTexts)
 	{
-		TSharedPtr<FTranslationPickerTextItem> Item = FTranslationPickerTextItem::BuildTextItem(PickedText, true);
+		TSharedPtr<FTranslationPickerTextItem> Item = FTranslationPickerTextItem::BuildTextItem(PickedText.Text, true);
 		
 		AllItems.Add(Item);
 
 		const FString& FilterBy = FilterText.ToString();
 
-		if (!PickedText.IsEmptyOrWhitespace() &&
-			!PickedText.ToString().Contains(FilterBy) &&
-			!PickedText.BuildSourceString().Contains(FilterBy))
+		if (!PickedText.Text.IsEmptyOrWhitespace() &&
+			!PickedText.Text.ToString().Contains(FilterBy) &&
+			!PickedText.Text.BuildSourceString().Contains(FilterBy))
 		{
 			continue;
 		}
diff --git a/Engine/Source/Developer/TranslationEditor/Private/TranslationPickerEditWindow.h b/Engine/Source/Developer/TranslationEditor/Private/TranslationPickerEditWindow.h
--- a/Engine/Source/Developer/TranslationEditor/Private/TranslationPickerEditWindow.h
+++ b/Engine/Source/Developer/TranslationEditor/Private/TranslationPickerEditWindow.h
@@ -3,7 +3,6 @@
 #pragma once
 
 #include "Containers/Array.h"
-#include "Containers/UnrealString.h"
 #include "HAL/Platform.h"
 #include "Input/Reply.h"
 #include "Internationalization/Text.h"
@@ -97,6 +96,14 @@ struct FTranslationPickerTextItem : public FGCObject
 	/** Create new text item */
 	static TSharedPtr<FTranslationPickerTextItem> BuildTextItem(const FText& InText, bool bAllowEditing);
 
+	inline bool operator==(const FTranslationPickerTextItem& Other) const
+	{
+		// It is sufficient to compare a subset to know they match. We can avoid comparing other fields.
+		return (CleanNamespace.Equals(Other.CleanNamespace) &&
+			TextId.GetKey() == Other.TextId.GetKey() &&
+			SourceString.Equals(Other.SourceString));
+	}
+
 	virtual void AddReferencedObjects(FReferenceCollector& Collector) override;
 	virtual FString GetReferencerName() const override
 	{
@@ -156,8 +163,6 @@ class STranslationPickerEditWindow : public SCompoundWidget
 
 	SLATE_ARGUMENT(TWeakPtr<SWindow>, ParentWindow)
 
-	SLATE_ARGUMENT(TArray<FText>, PickedTexts)
-
 	SLATE_END_ARGS()
 
 	virtual ~STranslationPickerEditWindow();
@@ -173,6 +178,8 @@ private:
 
 	FReply Close();
 
+	FReply Exit();
+
 	/** We need to support keyboard focus to process the 'Esc' key */
 	virtual bool SupportsKeyboardFocus() const override
 	{
@@ -182,8 +189,8 @@ private:
 	/** Return to picker floating window */
 	FReply RestorePicker();
 
-	/** Save all translations and close */
-	FReply SaveAllAndClose();
+	/** Save all translations and exit */
+	FReply SaveAllAndExit();
 
 	/** Update text list items */
 	void UpdateListItems();
@@ -208,9 +215,6 @@ private:
 	/** Contents of the window */
 	TSharedPtr<SBox> WindowContents;
 
-	/** The FTexts that we have found under the cursor */
-	TArray<FText> PickedTexts;
-
 	/** Full unfiltered list of items */
 	TArray<TSharedPtr<FTranslationPickerTextItem>> AllItems;
 
diff --git a/Engine/Source/Developer/TranslationEditor/Private/TranslationPickerFloatingWindow.cpp b/Engine/Source/Developer/TranslationEditor/Private/TranslationPickerFloatingWindow.cpp
--- a/Engine/Source/Developer/TranslationEditor/Private/TranslationPickerFloatingWindow.cpp
+++ b/Engine/Source/Developer/TranslationEditor/Private/TranslationPickerFloatingWindow.cpp
@@ -2,18 +2,16 @@
 
 #include "TranslationPickerFloatingWindow.h"
 
-#include "Containers/SortedMap.h"
-#include "Containers/UnrealString.h"
 #include "CoreGlobals.h"
 #include "Engine/Engine.h"
 #include "Engine/GameEngine.h"
 #include "Framework/Application/IInputProcessor.h"
 #include "Framework/Application/SlateApplication.h"
+#include "Framework/Application/SlateUser.h"
 #include "Framework/Docking/TabManager.h"
 #include "Framework/Text/TextLayout.h"
 #include "GameFramework/PlayerController.h"
 #include "HAL/Platform.h"
-#include "HAL/PlatformCrt.h"
 #include "HAL/PlatformApplicationMisc.h"
 #include "Input/Events.h"
 #include "InputCoreTypes.h"
@@ -28,20 +26,20 @@
 #include "Layout/ChildrenBase.h"
 #include "Layout/Margin.h"
 #include "Layout/SlateRect.h"
+#include "MathUtil.h"
 #include "Math/Vector2D.h"
-#include "Misc/Attribute.h"
 #if WITH_EDITOR
 #include "Editor.h"
 #include "SDocumentationToolTip.h"
 #endif // WITH_EDITOR
+#include "Rendering/DrawElements.h"
 #include "SlotBase.h"
+#include "Styling/WidgetStyle.h"
 #include "TranslationPickerEditWindow.h"
 #include "TranslationPickerWidget.h"
-#include "Types/SlateEnums.h"
+#include "Types/PaintArgs.h"
 #include "Widgets/IToolTip.h"
 #include "Widgets/Input/SEditableText.h"
-#include "Widgets/Layout/SBorder.h"
-#include "Widgets/Layout/SScrollBox.h"
 #include "Widgets/SBoxPanel.h"
 #include "Widgets/SToolTip.h"
 #include "Widgets/SWidget.h"
@@ -86,7 +84,7 @@ public:
 		if (Key == EKeys::Escape)
 		{
 			Owner->SetViewportMouseIgnoreLook(false);
-			Owner->Close();
+			Owner->Exit();
 			return true;
 		}
 		else if (Key == EKeys::Enter)
@@ -98,6 +96,11 @@ public:
 			}
 			return true;
 		}
+		else if (Key == EKeys::BackSpace)
+		{
+			TranslationPickerManager::bDrawBoxes = !TranslationPickerManager::bDrawBoxes;
+			return true;
+		}
 		else if (Key == EKeys::Backslash)
 		{
 			if (Owner->bMouseLookInputIgnored)
@@ -125,12 +128,12 @@ public:
 			FInputKeyManager::Get().GetCodesFromKey(EKeys::One, KeyCodeOne, CharCodeOne);
 
 			int32 EntryIndex = *CharCode - *CharCodeOne;
-			if (EntryIndex < 0 || EntryIndex > 4 || EntryIndex >= Owner->PickedTexts.Num())
+			if (EntryIndex < 0 || EntryIndex > 4 || EntryIndex >= TranslationPickerManager::PickedTexts.Num())
 			{
 				return false;	// Handle only first five entries, the max number of entries that fit in the floating picker
 			}
 
-			const FText& PickedText = Owner->PickedTexts[EntryIndex];
+			const FText& PickedText = TranslationPickerManager::PickedTexts[EntryIndex].Text;
 			
 			FTextId TextId = FTextInspector::GetTextId(PickedText);
 
@@ -179,7 +182,7 @@ void STranslationPickerFloatingWindow::Construct(const FArguments& InArgs)
 		.Padding(FMargin(5))
 		[
 			SNew(STextBlock)
-			.Text(PickedTexts.Num() > 0 ?
+			.Text(TranslationPickerManager::PickedTexts.Num() > 0 ?
 				LOCTEXT("TranslationPickerEnterToEdit", "Press Enter to edit translations") :
 				LOCTEXT("TranslationPickerHoverToViewEditEscToQuit", "Hover over text to view/edit translations, or press Esc to quit"))
 			.Justification(ETextJustify::Center)
@@ -220,30 +223,45 @@ FReply STranslationPickerFloatingWindow::Close()
 	return FReply::Handled();
 }
 
+FReply STranslationPickerFloatingWindow::Exit()
+{
+	TranslationPickerManager::RemoveOverlay();
+	Close();
+
+	return FReply::Handled();
+}
+
 void STranslationPickerFloatingWindow::Tick( const FGeometry& AllottedGeometry, const double InCurrentTime, const float InDeltaTime )
 {
-	FWidgetPath Path = FSlateApplication::Get().LocateWindowUnderMouse(FSlateApplication::Get().GetCursorPos(), FSlateApplication::Get().GetInteractiveTopLevelWindows(), true);
+	const FVector2f MousePos = FSlateApplication::Get().GetCursorPos();
+	float MousePosDiffX = FMath::Abs(MousePos.X - MousePosPrev.X);
+	float MousePosDiffY = FMath::Abs(MousePos.Y - MousePosPrev.Y);
 
-	if (Path.IsValid())
+	if (MousePosDiffX >= 1.0f || MousePosDiffY >= 1.0f)
 	{
-		// If the path of widgets we're hovering over changed since last time (or if this is the first tick and LastTickHoveringWidgetPath hasn't been set yet)
-		if (!LastTickHoveringWidgetPath.IsValid() || LastTickHoveringWidgetPath.ToWidgetPath().ToString() != Path.ToString())
-		{
-			// Clear all previous text and widgets
-			PickedTexts.Reset();
+		MousePosPrev = MousePos;
 
-			// Process the leaf-widget under the cursor
-			if (Path.Widgets.Num() > 0)
+		FWidgetPath Path = FSlateApplication::Get().LocateWindowUnderMouse(MousePos, FSlateApplication::Get().GetInteractiveTopLevelWindows(), true);
+
+		if (Path.IsValid())
+		{
+			// If the path of widgets we're hovering over changed since last time (or if this is the first tick and LastTickHoveringWidgetPath hasn't been set yet)
+			if (!LastTickHoveringWidgetPath.IsValid() || LastTickHoveringWidgetPath.ToWidgetPath().ToString() != Path.ToString())
 			{
-				// General Widget case
+				// Clear all previous text and widgets
+				TranslationPickerManager::PickedTexts.Reset();
+
+				// Process the leaf-widget under the cursor
 				TSharedRef<SWidget> PathWidget = Path.Widgets.Last().Widget;
-				PickTextFromWidget(PathWidget);
+
+				// General Widget case
+				PickTextFromWidget(PathWidget, Path, false);
 
 				// Tooltip case
 				TSharedPtr<IToolTip> Tooltip = PathWidget->GetToolTip();
 				if (Tooltip.IsValid() && !Tooltip->IsEmpty())
 				{
-					PickTextFromWidget(Tooltip->AsWidget());
+					PickTextFromWidget(Tooltip->AsWidget(), Path, true);
 				}
 
 				// Also include tooltips from parent widgets in this path (since they may be visible)
@@ -255,13 +273,56 @@ void STranslationPickerFloatingWindow::Tick( const FGeometry& AllottedGeometry,
 					TSharedPtr<IToolTip> ParentTooltip = ParentPathWidget->GetToolTip();
 					if (ParentTooltip.IsValid() && !ParentTooltip->IsEmpty())
 					{
-						PickTextFromWidget(ParentTooltip->AsWidget());
+						PickTextFromWidget(ParentTooltip->AsWidget(), Path, true);
 					}
 				}
 			}
 
+			TranslationPickerManager::PickedTexts.Sort([this, MousePos](const FTranslationPickerTextAndGeom& LHS, const FTranslationPickerTextAndGeom& RHS)
+			{
+				FSlateRect RectLHS = GetRect(LHS.Geometry);
+				FSlateRect RectRHS = GetRect(RHS.Geometry);
+				if (IsNearlyEqual(RectLHS, RectRHS))
+				{
+					FString SourceStringLHS;
+					FString SourceStringRHS;
+					if (const FString* SourceStringPtrLHS = FTextInspector::GetSourceString(LHS.Text))
+					{
+						SourceStringLHS = *SourceStringPtrLHS;
+					}
+					if (const FString* SourceStringPtrRHS = FTextInspector::GetSourceString(RHS.Text))
+					{
+						SourceStringRHS = *SourceStringPtrRHS;
+					}
+					return (SourceStringLHS < SourceStringRHS);
+				}
+
+				bool bContainsLHS = false;
+				bool bContainsRHS = false;
+				float DistLHS = DistSquaredToRect(RectLHS, MousePos, bContainsLHS);
+				float DistRHS = DistSquaredToRect(RectRHS, MousePos, bContainsRHS);
+
+				if ((bContainsLHS && bContainsRHS) ||
+					(!bContainsLHS && !bContainsRHS))
+				{
+					return DistLHS < DistRHS;
+				}
+				else if (bContainsLHS && !bContainsRHS)
+				{
+					return true;
+				}
+				else if (!bContainsLHS && bContainsRHS)
+				{
+					return false;
+				}
+				ensure(false);		// All cases handled above
+				return true;
+			});
+
 			UpdateListItems();
 		}
+
+		LastTickHoveringWidgetPath = FWeakWidgetPath(Path);
 	}
 
 	if (ParentWindow.IsValid())
@@ -271,7 +332,7 @@ void STranslationPickerFloatingWindow::Tick( const FGeometry& AllottedGeometry,
 		DesiredPosition.X -= FSlateApplication::Get().GetCursorSize().X;
 		DesiredPosition.Y += FSlateApplication::Get().GetCursorSize().Y;
 
-		// Move to opposite side of the cursor than the tool tip, so they don't overlaps
+		// Move to opposite side of the cursor than the tool tip, so they don't overlap
 		DesiredPosition.X -= WindowSize.X;
 
 		// Clamp to work area
@@ -280,55 +341,141 @@ void STranslationPickerFloatingWindow::Tick( const FGeometry& AllottedGeometry,
 		// also kind of a hack, but this is the only way at the moment to get a 'cursor decorator' without using the drag-drop code path
 		ParentWindow.Pin()->MoveWindowTo(DesiredPosition);
 	}
+}
+
+FVector2f STranslationPickerFloatingWindow::GetNearestPoint(const FSlateRect& Rect, const FVector2f& Point, bool& Contains)
+{
+	FVector2f Result = Point;
+	Contains = true;
+
+	// Check each dimension individually, clamping the point to the nearest edge endpoint
+	if (Point.X < Rect.Left)
+	{
+		Result.X = Rect.Left;
+		Contains = false;
+	}
+	if (Point.X > Rect.Right)
+	{
+		Result.X = Rect.Right;
+		Contains = false;
+	}
+	if (Point.Y < Rect.Top)
+	{
+		Result.Y = Rect.Top;
+		Contains = false;
+	}
+	if (Point.Y > Rect.Bottom)
+	{
+		Result.Y = Rect.Bottom;
+		Contains = false;
+	}
+
+	if (!Contains)
+	{
+		return Result;
+	}
+
+	// Point is inside the rectangle, find the position on the nearest edge
+	float DistTop =		FMath::Abs(Point.Y - Rect.Top);
+	float DistBottom =	FMath::Abs(Point.Y - Rect.Bottom);
+	float DistLeft =	FMath::Abs(Point.X - Rect.Left);
+	float DistRight =	FMath::Abs(Point.X - Rect.Right);
+	if (DistTop <= DistBottom && DistTop <= DistLeft && DistTop <= DistRight)
+	{
+		Result.Y = Rect.Top;
+	}
+	else if (DistBottom <= DistLeft && DistBottom <= DistRight)
+	{
+		Result.Y = Rect.Bottom;
+	}
+	else if (DistLeft <= DistRight)
+	{
+		Result.X = Rect.Left;
+	}
+	else
+	{
+		Result.X = Rect.Right;
+	}
+
+	return Result;
+}
+
+float STranslationPickerFloatingWindow::DistSquaredToRect(const FSlateRect& Rect, const FVector2f& Point, bool& Contains)
+{
+	if ((Rect.GetSize().X < FMathf::ZeroTolerance) ||
+		(Rect.GetSize().Y < FMathf::ZeroTolerance))
+	{
+		Contains = false;
+		return FLT_MAX;
+	}
+
+	FVector2f NearestPointOnRect = GetNearestPoint(Rect, Point, Contains);
+
+	return FVector2f::DistSquared(NearestPointOnRect, Point);
+}
 
-	LastTickHoveringWidgetPath = FWeakWidgetPath(Path);
+bool STranslationPickerFloatingWindow::IsNearlyEqual(const FSlateRect& RectLHS, const FSlateRect& RectRHS)
+{
+	return FMath::IsNearlyEqual(RectLHS.Left, RectRHS.Left) &&
+		FMath::IsNearlyEqual(RectLHS.Top, RectRHS.Top) &&
+		FMath::IsNearlyEqual(RectLHS.Right, RectRHS.Right) &&
+		FMath::IsNearlyEqual(RectLHS.Bottom, RectRHS.Bottom);
+}
+
+FSlateRect STranslationPickerFloatingWindow::GetRect(const FPaintGeometry& Geometry)
+{
+	FSlateRenderTransform TransformPaint = Geometry.GetAccumulatedRenderTransform();
+	FVector2f Pos = FVector2f(TransformPaint.GetTranslation());
+	FVector2f Size = FVector2f(Geometry.GetLocalSize()) * TransformPaint.GetMatrix().GetScale().GetVector();
+
+	return FSlateRect(Pos.X, Pos.Y, Pos.X + Size.X, Pos.Y + Size.Y);
 }
 
-void STranslationPickerFloatingWindow::PickTextFromWidget(TSharedRef<SWidget> Widget)
+void STranslationPickerFloatingWindow::PickTextFromWidget(TSharedRef<SWidget> Widget, const FWidgetPath& Path, bool IsToolTip)
 {
-	auto AppendPickedTextImpl = [this](const FText& InPickedText)
+	auto AppendPickedTextImpl = [this](const FText& InText, TSharedPtr<SWidget> InWidget, const FWidgetPath& InPath, bool IsToolTip)
 	{
-		const bool bAlreadyPicked = PickedTexts.ContainsByPredicate([&InPickedText](const FText& InOtherPickedText)
+		const bool bAlreadyPicked = TranslationPickerManager::PickedTexts.ContainsByPredicate([&InText](const FTranslationPickerTextAndGeom& inOther)
 		{
-			return InOtherPickedText.IdenticalTo(InPickedText);
+			return inOther.Text.IdenticalTo(InText);
 		});
 
 		if (!bAlreadyPicked)
 		{
-			PickedTexts.Add(InPickedText);
+			TranslationPickerManager::PickedTexts.Emplace(InText, GetPaintGeometry(InWidget, InPath, IsToolTip));
 		}
 	};
 
-	auto AppendPickedText = [this, AppendPickedTextImpl](const FText& InPickedText)
+	auto AppendPickedText = [this, AppendPickedTextImpl](const FText& InText, TSharedPtr<SWidget> InWidget, const FWidgetPath& InPath, bool IsToolTip)
 	{
-		if (InPickedText.IsEmpty())
+		if (InText.IsEmpty())
 		{
 			return;
 		}
 
 		// Search the text from this widget's FText::Format history to find any source text
 		TArray<FHistoricTextFormatData> HistoricFormatData;
-		FTextInspector::GetHistoricFormatData(InPickedText, HistoricFormatData);
+		FTextInspector::GetHistoricFormatData(InText, HistoricFormatData);
 
 		if (HistoricFormatData.Num() > 0)
 		{
 			for (const FHistoricTextFormatData& HistoricFormatDataItem : HistoricFormatData)
 			{
-				AppendPickedTextImpl(HistoricFormatDataItem.SourceFmt.GetSourceText());
+				AppendPickedTextImpl(HistoricFormatDataItem.SourceFmt.GetSourceText(), InWidget, InPath, IsToolTip);
 
 				for (auto It = HistoricFormatDataItem.Arguments.CreateConstIterator(); It; ++It)
 				{
 					const FFormatArgumentValue& ArgumentValue = It.Value();
 					if (ArgumentValue.GetType() == EFormatArgumentType::Text)
 					{
-						AppendPickedTextImpl(ArgumentValue.GetTextValue());
+						AppendPickedTextImpl(ArgumentValue.GetTextValue(), InWidget, InPath, IsToolTip);
 					}
 				}
 			}
 		}
 		else
 		{
-			AppendPickedTextImpl(InPickedText);
+			AppendPickedTextImpl(InText, InWidget, InPath, IsToolTip);
 		}
 	};
 
@@ -336,75 +483,83 @@ void STranslationPickerFloatingWindow::PickTextFromWidget(TSharedRef<SWidget> Wi
 	if (Widget->GetTypeAsString() == "STextBlock")
 	{
 		STextBlock& TextBlock = (STextBlock&)Widget.Get();
-		AppendPickedText(TextBlock.GetText());
+		AppendPickedText(TextBlock.GetText(), Widget.ToSharedPtr(), Path, IsToolTip);
 	}
 	else if (Widget->GetTypeAsString() == "SRichTextBlock")
 	{
 		SRichTextBlock& RichTextBlock = (SRichTextBlock&)Widget.Get();
-		AppendPickedText(RichTextBlock.GetText());
+		AppendPickedText(RichTextBlock.GetText(), Widget.ToSharedPtr(), Path, IsToolTip);
 	}
 	else if (Widget->GetTypeAsString() == "SToolTip")
 	{
 		SToolTip& ToolTipWidget = (SToolTip&)Widget.Get();
-		AppendPickedText(ToolTipWidget.GetTextTooltip());
+		AppendPickedText(ToolTipWidget.GetTextTooltip(), Widget.ToSharedPtr(), Path, IsToolTip);
 	}
 #if WITH_EDITOR
 	else if (Widget->GetTypeAsString() == "SDocumentationToolTip")
 	{
 		SDocumentationToolTip& DocumentationToolTip = (SDocumentationToolTip&)Widget.Get();
-		AppendPickedText(DocumentationToolTip.GetTextTooltip());
+		AppendPickedText(DocumentationToolTip.GetTextTooltip(), Widget.ToSharedPtr(), Path, IsToolTip);
 	}
 #endif // WITH_EDITOR
 	else if (Widget->GetTypeAsString() == "SEditableText")
 	{
 		SEditableText& EditableText = (SEditableText&)Widget.Get();
-		AppendPickedText(EditableText.GetText());
-		AppendPickedText(EditableText.GetHintText());
+		AppendPickedText(EditableText.GetText(), Widget.ToSharedPtr(), Path, IsToolTip);
+		AppendPickedText(EditableText.GetHintText(), Widget.ToSharedPtr(), Path, IsToolTip);
 	}
 	else if (Widget->GetTypeAsString() == "SMultiLineEditableText")
 	{
 		SMultiLineEditableText& MultiLineEditableText = (SMultiLineEditableText&)Widget.Get();
-		AppendPickedText(MultiLineEditableText.GetText());
-		AppendPickedText(MultiLineEditableText.GetHintText());
+		AppendPickedText(MultiLineEditableText.GetText(), Widget.ToSharedPtr(), Path, IsToolTip);
+		AppendPickedText(MultiLineEditableText.GetHintText(), Widget.ToSharedPtr(), Path, IsToolTip);
 	}
 
 	// Recurse into child widgets
-	PickTextFromChildWidgets(Widget);
+	PickTextFromChildWidgets(Widget, Path, IsToolTip);
 }
 
-void STranslationPickerFloatingWindow::PickTextFromChildWidgets(TSharedRef<SWidget> Widget)
+void STranslationPickerFloatingWindow::PickTextFromChildWidgets(TSharedRef<SWidget> Widget, const FWidgetPath& Path, bool IsToolTip)
 {
 	FChildren* Children = Widget->GetChildren();
 
+	FWidgetPath PathChild = Path;
+
 	for (int32 ChildIndex = 0; ChildIndex < Children->Num(); ++ChildIndex)
 	{
 		TSharedRef<SWidget> ChildWidget = Children->GetChildAt(ChildIndex);
 
+		bool bExtended = PathChild.ExtendPathTo(FWidgetMatcher(ChildWidget), EVisibility::Visible);
+
 		// Pull out any FText from this child widget
-		PickTextFromWidget(ChildWidget);
+		PickTextFromWidget(ChildWidget, PathChild, IsToolTip);
 	}
 }
 
 bool STranslationPickerFloatingWindow::SwitchToEditWindow()
 {
-	if (PickedTexts.Num() > 0)
+	if (TranslationPickerManager::PickedTexts.Num() == 0)
+	{
+		return false;
+	}
+
+	TSharedPtr<SWindow> RootWindow = FGlobalTabmanager::Get()->GetRootWindow();
+
+	// Open a different window to allow editing of the translation
 	{
-		// Open a different window to allow editing of the translation
 		TSharedRef<SWindow> NewWindow = SNew(SWindow)
 			.Title(LOCTEXT("TranslationPickerEditWindowTitle", "Edit Translations"))
 			.CreateTitleBar(true)
 			.SizingRule(ESizingRule::UserSized);
 
 		TSharedRef<STranslationPickerEditWindow> EditWindow = SNew(STranslationPickerEditWindow)
-			.ParentWindow(NewWindow)
-			.PickedTexts(PickedTexts);
+			.ParentWindow(NewWindow);
 
 		NewWindow->SetContent(EditWindow);
 
 		// Make this roughly the same size as the Edit Window, so when you press Esc to edit, the window is in basically the same size
-		NewWindow->Resize(FVector2D(STranslationPickerEditWindow::DefaultEditWindowWidth, STranslationPickerEditWindow::DefaultEditWindowHeight));
+		NewWindow->Resize(FVector2f(STranslationPickerEditWindow::DefaultEditWindowWidth, STranslationPickerEditWindow::DefaultEditWindowHeight));
 
-		TSharedPtr<SWindow> RootWindow = FGlobalTabmanager::Get()->GetRootWindow();
 		if (RootWindow.IsValid())
 		{
 			FSlateApplication::Get().AddWindowAsNativeChild(NewWindow, RootWindow.ToSharedRef());
@@ -413,12 +568,10 @@ bool STranslationPickerFloatingWindow::SwitchToEditWindow()
 		{
 			FSlateApplication::Get().AddWindow(NewWindow);
 		}
-
 		NewWindow->MoveWindowTo(ParentWindow.Pin()->GetPositionInScreen());
-
-		return true;
 	}
-	return false;
+
+	return true;
 }
 
 void STranslationPickerFloatingWindow::SetViewportMouseIgnoreLook(bool bLookIgnore)
@@ -466,9 +619,9 @@ void STranslationPickerFloatingWindow::UpdateListItems()
 {
 	TextListItems.Reset();
 
-	for (const FText& PickedText : PickedTexts)
+	for (const FTranslationPickerTextAndGeom& PickedText : TranslationPickerManager::PickedTexts)
 	{
-		TSharedPtr<FTranslationPickerTextItem> Item = FTranslationPickerTextItem::BuildTextItem(PickedText, false);
+		TSharedPtr<FTranslationPickerTextItem> Item = FTranslationPickerTextItem::BuildTextItem(PickedText.Text, false);
 
 		TextListItems.Add(Item);
 	}
@@ -485,5 +638,109 @@ TSharedRef<ITableRow> STranslationPickerFloatingWindow::TextListView_OnGenerateW
 	return SNew(STranslationPickerEditWidget, OwnerTable, InItem);
 }
 
+FPaintGeometry STranslationPickerFloatingWindow::GetPaintGeometry(const TSharedPtr<SWidget>& PickedWidget, const FWidgetPath& PickedPath, bool IsToolTip) const
+{
+	FVector2f OffsetTooltip(0.0f);
+	FVector2f OffsetMenu(0.0f);
+
+	if (IsToolTip)
+	{
+		const TSharedPtr<FSlateUser> SlateUser = FSlateApplication::Get().GetUser(0);
+		OffsetTooltip = SlateUser->GetTooltipPosition();
+	}
+	else if (PickedPath.GetDeepestWindow()->GetType() == EWindowType::Menu)
+	{
+		const FArrangedWidget& ArrangedWidget = PickedPath.Widgets[0];
+		FPaintGeometry GeomPaint = PickedPath.Widgets[0].Geometry.ToPaintGeometry();
+		const FSlateRenderTransform& TransformPaint = GeomPaint.GetAccumulatedRenderTransform();
+		OffsetMenu = FVector2f(TransformPaint.GetTranslation());
+	}
+
+	FPaintGeometry PaintGeometry;
+	if (IsToolTip)
+	{
+		FSlateLayoutTransform TransformTooltip = FSlateLayoutTransform(OffsetTooltip);
+		PaintGeometry.AppendTransform(TransformTooltip);
+	}
+	else
+	{
+		FSlateLayoutTransform TransformMenu = FSlateLayoutTransform(OffsetMenu);
+
+		TSharedRef<SWidget> PickedWidgetRef = PickedWidget.ToSharedRef();
+		bool bFound = GetGeometry(PickedWidgetRef, PaintGeometry, TransformMenu);
+		if (!bFound)
+		{
+			PaintGeometry = PickedWidget->GetPaintSpaceGeometry().ToPaintGeometry(TransformMenu);
+		}
+	}
+
+	return PaintGeometry;
+}
+
+bool STranslationPickerFloatingWindow::GetGeometry(const TSharedRef<const SWidget>& Widget, FPaintGeometry& PaintGeometry, const FSlateLayoutTransform& TransformOffset) const
+{
+	TSharedPtr<SWindow> WidgetWindow = FSlateApplication::Get().FindWidgetWindow(Widget);
+	if (!WidgetWindow.IsValid())
+	{
+		return false;
+	}
+
+	while (WidgetWindow->GetParentWidget().IsValid())
+	{
+		TSharedRef<SWidget> CurrentWidget = WidgetWindow->GetParentWidget().ToSharedRef();
+		TSharedPtr<SWindow> ParentWidgetWindow = FSlateApplication::Get().FindWidgetWindow(CurrentWidget);
+		if (!ParentWidgetWindow.IsValid())
+		{
+			break;
+		}
+		WidgetWindow = ParentWidgetWindow;
+	}
+
+	TSharedRef<SWindow> CurrentWindowRef = WidgetWindow.ToSharedRef();
+
+	FWidgetPath WidgetPath;
+	if (FSlateApplication::Get().GeneratePathToWidgetUnchecked(Widget, WidgetPath))
+	{
+		FArrangedWidget ArrangedWidget = WidgetPath.FindArrangedWidget(Widget).Get(FArrangedWidget::GetNullWidget());
+		ArrangedWidget.Geometry.AppendTransform(FSlateLayoutTransform(Inverse(CurrentWindowRef->GetPositionInScreen())));
+		ArrangedWidget.Geometry.AppendTransform(TransformOffset);
+
+		const FVector2D InflateAmount = FVector2D(1, 1) / FVector2D(ArrangedWidget.Geometry.GetAccumulatedRenderTransform().GetMatrix().GetScale().GetVector());
+		PaintGeometry = ArrangedWidget.Geometry.ToInflatedPaintGeometry(InflateAmount);
+		return true;
+	}
+
+	return false;
+}
+
+int32 STranslationPickerOverlay::OnPaint(const FPaintArgs& Args, const FGeometry& AllottedGeometry, const FSlateRect& MyCullingRect, FSlateWindowElementList& OutDrawElements, int32 LayerId, const FWidgetStyle& InWidgetStyle, bool bParentEnabled) const
+{
+	if (!TranslationPickerManager::bDrawBoxes)
+	{
+		return LayerId;
+	}
+
+	LayerId = SCompoundWidget::OnPaint(Args, AllottedGeometry, MyCullingRect, OutDrawElements, LayerId, InWidgetStyle, bParentEnabled);
+
+	static const FName DebugBorderBrush = TEXT("Debug.Border");
+	const FSlateBrush* Brush = FCoreStyle::Get().GetBrush(DebugBorderBrush);
+	const FLinearColor BoxColorYellow = FLinearColor::Yellow;														
+
+	const TArray<FTranslationPickerTextAndGeom>& PickedTexts = TranslationPickerManager::PickedTexts;
+
+	for (int32 i = 0; i < PickedTexts.Num(); ++i)
+	{
+		FSlateDrawElement::MakeBox(
+			OutDrawElements,
+			LayerId,
+			PickedTexts[i].Geometry,
+			Brush,
+			ESlateDrawEffect::None,
+			BoxColorYellow);
+	}
+
+	return LayerId;
+}
+
 #undef LOCTEXT_NAMESPACE
 
diff --git a/Engine/Source/Developer/TranslationEditor/Private/TranslationPickerFloatingWindow.h b/Engine/Source/Developer/TranslationEditor/Private/TranslationPickerFloatingWindow.h
--- a/Engine/Source/Developer/TranslationEditor/Private/TranslationPickerFloatingWindow.h
+++ b/Engine/Source/Developer/TranslationEditor/Private/TranslationPickerFloatingWindow.h
@@ -3,13 +3,19 @@
 #pragma once
 
 #include "Containers/Array.h"
+#include "Layout/PaintGeometry.h"
 #include "Layout/WidgetPath.h"
 #include "Templates/SharedPointer.h"
 #include "Widgets/DeclarativeSyntaxSupport.h"
 #include "Widgets/SCompoundWidget.h"
 #include "Widgets/Views/SListView.h"
 
+class FPaintArgs;
+class FSlateRect;
+class FSlateWindowElementList;
 class FText;
+class FTranslationPickerInputProcessor;
+class FWidgetStyle;
 class ITableRow;
 class STableViewBase;
 class SToolTip;
@@ -20,8 +26,6 @@ struct FTranslationPickerTextItem;
 
 #define LOCTEXT_NAMESPACE "TranslationPicker"
 
-class FTranslationPickerInputProcessor;
-
 /** Translation picker floating window to show details of FText(s) under cursor, and allow in-place translation via TranslationPickerEditWindow */
 class STranslationPickerFloatingWindow : public SCompoundWidget
 {
@@ -40,17 +44,28 @@ private:
 
 	FReply Close();
 
+	FReply Exit();
+
 	virtual void Tick( const FGeometry& AllottedGeometry, const double InCurrentTime, const float InDeltaTime ) override;
 
 	/** Pull the FText reference out of an SWidget */
-	void PickTextFromWidget(TSharedRef<SWidget> Widget);
+	void PickTextFromWidget(TSharedRef<SWidget> Widget, const FWidgetPath& Path, bool IsToolTip);
 
 	/** Pull the FText reference out of the child widgets of an SWidget */
-	void PickTextFromChildWidgets(TSharedRef<SWidget> Widget);
+	void PickTextFromChildWidgets(TSharedRef<SWidget> Widget, const FWidgetPath& Path, bool IsToolTip);
 
 	/** Switch from floating window to edit window */
 	bool SwitchToEditWindow();
 
+	/** Find the nearest point on the rectangle, and whether it is contained in the rect */
+	static FVector2f GetNearestPoint(const FSlateRect& Rect, const FVector2f& Point, bool& Contains);
+
+	static float DistSquaredToRect(const FSlateRect& Rect, const FVector2f& Point, bool& Contains);
+
+	static bool IsNearlyEqual(const FSlateRect& RectLHS, const FSlateRect& RectRHS);
+
+	static FSlateRect GetRect(const FPaintGeometry& Geometry);
+
 	/** Update text list items */
 	void UpdateListItems();
 
@@ -60,6 +75,10 @@ private:
 	/** Get world from editor or engine */
 	UWorld* GetWorld() const;
 
+	FPaintGeometry GetPaintGeometry(const TSharedPtr<SWidget>& PickedWidget, const FWidgetPath& PickedPath, bool IsToolTip) const;
+
+	bool GetGeometry(const TSharedRef<const SWidget>& Widget, FPaintGeometry& PaintGeometry, const FSlateLayoutTransform& TransformOffset) const;
+
 	TSharedRef<ITableRow> TextListView_OnGenerateWidget(TSharedPtr<FTranslationPickerTextItem> InItem, const TSharedRef<STableViewBase>& OwnerTable);
 
 	/** Input processor used to capture key and mouse events */
@@ -71,9 +90,6 @@ private:
 	/** Contents of the window */
 	TSharedPtr<SToolTip> WindowContents;
 
-	/** The FTexts that we have found under the cursor */
-	TArray<FText> PickedTexts;
-
 	/** List items for the text list */
 	TArray<TSharedPtr<FTranslationPickerTextItem>> TextListItems;
 
@@ -81,10 +97,28 @@ private:
 	typedef SListView<TSharedPtr<FTranslationPickerTextItem>> STextListView;
 	TSharedPtr<STextListView> TextListView;
 
+	FVector2f MousePosPrev;
+
 	/** The path widgets we were hovering over last tick */
 	FWeakWidgetPath LastTickHoveringWidgetPath;
 
 	bool bMouseLookInputIgnored = false;
 };
 
+class STranslationPickerOverlay : public SCompoundWidget
+{
+public:
+	SLATE_BEGIN_ARGS(STranslationPickerOverlay)
+	{
+		_Visibility = EVisibility::HitTestInvisible;
+	}
+
+	SLATE_END_ARGS()
+
+	void Construct(const FArguments& InArgs) {}
+
+private:
+	virtual int32 OnPaint(const FPaintArgs& Args, const FGeometry& AllottedGeometry, const FSlateRect& MyCullingRect, FSlateWindowElementList& OutDrawElements, int32 LayerId, const FWidgetStyle& InWidgetStyle, bool bParentEnabled) const override;
+};
+
 #undef LOCTEXT_NAMESPACE
diff --git a/Engine/Source/Developer/TranslationEditor/Private/TranslationPickerWidget.cpp b/Engine/Source/Developer/TranslationEditor/Private/TranslationPickerWidget.cpp
--- a/Engine/Source/Developer/TranslationEditor/Private/TranslationPickerWidget.cpp
+++ b/Engine/Source/Developer/TranslationEditor/Private/TranslationPickerWidget.cpp
@@ -2,15 +2,14 @@
 
 #include "TranslationPickerWidget.h"
 
+#include "Engine/GameEngine.h"
 #include "Framework/Application/SlateApplication.h"
 #include "Framework/Docking/TabManager.h"
 #include "Internationalization/Internationalization.h"
 #include "Internationalization/Text.h"
 #include "Layout/Children.h"
-#include "Layout/Margin.h"
 #include "Math/MathFwd.h"
 #include "Math/Vector2D.h"
-#include "Misc/Attribute.h"
 #include "SlotBase.h"
 #include "Styling/AppStyle.h"
 #include "Styling/ISlateStyle.h"
@@ -25,6 +24,9 @@
 #include "Widgets/SToolTip.h"
 #include "Widgets/SWindow.h"
 #include "Widgets/Text/STextBlock.h"
+#if WITH_EDITOR
+#include "Editor.h"
+#endif // WITH_EDITOR
 
 class SWidget;
 
@@ -123,27 +125,34 @@ void STranslationWidgetPicker::OnCheckStateChanged(const ECheckBoxState NewCheck
 	}
 }
 
-TSharedPtr<SWindow> TranslationPickerManager::PickerWindow = TSharedPtr<SWindow>();
-TSharedPtr<STranslationPickerFloatingWindow> TranslationPickerManager::PickerWindowWidget = TSharedPtr<STranslationPickerFloatingWindow>();
+TSharedPtr<SWindow> TranslationPickerManager::PickerWindow;
+TSharedPtr<STranslationPickerFloatingWindow> TranslationPickerManager::PickerWindowWidget;
+TSharedPtr<STranslationPickerOverlay> TranslationPickerManager::MainWindowOverlay;
+TArray<FTranslationPickerTextAndGeom> TranslationPickerManager::PickedTexts;
+bool TranslationPickerManager::bDrawBoxes = true;
 
 bool TranslationPickerManager::OpenPickerWindow()
 {
+	if (PickerWindow.IsValid() || PickerWindowWidget.IsValid())
+	{
+		return false;
+	}
+
+	TSharedPtr<SWindow> RootWindow = FGlobalTabmanager::Get()->GetRootWindow();
+
 	// Not picking previously, launch a picker window
-	if (!PickerWindow.IsValid() && !PickerWindowWidget.IsValid())
 	{
 		TSharedRef<SWindow> NewWindow = SWindow::MakeCursorDecorator();
 		NewWindow->SetSizingRule(ESizingRule::FixedSize);
 		// The Edit window and Floating window should be roughly the same size, so it isn't too distracting switching between them
-		NewWindow->Resize(FVector2D(STranslationPickerEditWindow::DefaultEditWindowWidth, STranslationPickerEditWindow::DefaultEditWindowHeight));
+		NewWindow->Resize(FVector2f(STranslationPickerEditWindow::DefaultEditWindowWidth, STranslationPickerEditWindow::DefaultEditWindowHeight));
 		NewWindow->MoveWindowTo(FSlateApplication::Get().GetCursorPos());
 		PickerWindow = NewWindow;
 
 		NewWindow->SetContent(
 			SAssignNew(PickerWindowWidget, STranslationPickerFloatingWindow)
-			.ParentWindow(NewWindow)
-			);
+			.ParentWindow(NewWindow));
 
-		TSharedPtr<SWindow> RootWindow = FGlobalTabmanager::Get()->GetRootWindow();
 		if (RootWindow.IsValid())
 		{
 			FSlateApplication::Get().AddWindowAsNativeChild(NewWindow, RootWindow.ToSharedRef());
@@ -152,11 +161,36 @@ bool TranslationPickerManager::OpenPickerWindow()
 		{
 			FSlateApplication::Get().AddWindow(NewWindow);
 		}
+	}
+
+	// Create overlay for drawing text boxes
+	if (!MainWindowOverlay.IsValid())
+	{
+		MainWindowOverlay = SNew(STranslationPickerOverlay);
 
-		return true;
+#if WITH_EDITOR
+		if (GIsEditor && IsValid(GEditor))
+		{
+			RootWindow->AddOverlaySlot()
+			[
+				MainWindowOverlay.ToSharedRef()
+			];
+		}
+		else
+#endif // WITH_EDITOR
+		if (UGameEngine* GameEngine = Cast<UGameEngine>(GEngine))
+		{
+			if (TSharedPtr<SWindow> GameViewportWindow = GameEngine->GameViewportWindow.Pin())
+			{
+				GameViewportWindow->AddOverlaySlot()
+				[
+					MainWindowOverlay.ToSharedRef()
+				];
+			}
+		}
 	}
 
-	return false;
+	return true;
 }
 
 void TranslationPickerManager::ClosePickerWindow()
@@ -176,4 +210,27 @@ void TranslationPickerManager::ResetPickerWindow()
 	PickerWindowWidget.Reset();
 }
 
+void TranslationPickerManager::RemoveOverlay()
+{
+	if (!MainWindowOverlay.IsValid())
+	{
+		return;
+	}
+
+#if WITH_EDITOR
+	TSharedPtr<SWindow> RootWindow = FGlobalTabmanager::Get()->GetRootWindow();
+	RootWindow->RemoveOverlaySlot(MainWindowOverlay.ToSharedRef());
+	MainWindowOverlay.Reset();
+#else
+	if (UGameEngine* GameEngine = Cast<UGameEngine>(GEngine))
+	{
+		if (TSharedPtr<SWindow> GameViewportWindow = GameEngine->GameViewportWindow.Pin())
+		{
+			GameViewportWindow->RemoveOverlaySlot(MainWindowOverlay.ToSharedRef());
+			MainWindowOverlay.Reset();
+		}
+	}
+#endif // WITH_EDITOR
+}
+
 #undef LOCTEXT_NAMESPACE
diff --git a/Engine/Source/Developer/TranslationEditor/Private/TranslationPickerWidget.h b/Engine/Source/Developer/TranslationEditor/Private/TranslationPickerWidget.h
--- a/Engine/Source/Developer/TranslationEditor/Private/TranslationPickerWidget.h
+++ b/Engine/Source/Developer/TranslationEditor/Private/TranslationPickerWidget.h
@@ -8,7 +8,9 @@
 
 #define LOCTEXT_NAMESPACE "TranslationPicker"
 
+struct FTranslationPickerTextAndGeom;
 class STranslationPickerFloatingWindow;
+class STranslationPickerOverlay;
 class SWindow;
 
 enum class ECheckBoxState : uint8;
@@ -18,6 +20,13 @@ class TranslationPickerManager
 public:
 	static TSharedPtr<SWindow> PickerWindow;
 	static TSharedPtr<STranslationPickerFloatingWindow> PickerWindowWidget;
+	static TSharedPtr<STranslationPickerOverlay> MainWindowOverlay;
+
+	/** The FTexts that we have found under the cursor */
+	static TArray<FTranslationPickerTextAndGeom> PickedTexts;
+
+	/** Whether to draw boxes in the overlay */
+	static bool bDrawBoxes;
 
 	static bool IsPickerWindowOpen() { return PickerWindow.IsValid(); }
 
@@ -26,6 +35,8 @@ public:
 	static void ClosePickerWindow();
 
 	static void ResetPickerWindow();
+
+	static void RemoveOverlay();
 };
 
 /** Widget used to launch a 'picking' session */
@@ -53,4 +64,16 @@ private:
 	void OnCheckStateChanged(const ECheckBoxState NewCheckedState);
 };
 
+struct FTranslationPickerTextAndGeom
+{
+	FText Text;
+	FPaintGeometry Geometry;
+
+	FTranslationPickerTextAndGeom(const FText& InText, const FPaintGeometry& InGeometry)
+		: Text(InText)
+		, Geometry(InGeometry)
+	{
+	}
+};
+
 #undef LOCTEXT_NAMESPACE
diff --git a/Engine/Source/Runtime/Slate/Private/Framework/Application/SlateUser.cpp b/Engine/Source/Runtime/Slate/Private/Framework/Application/SlateUser.cpp
--- a/Engine/Source/Runtime/Slate/Private/Framework/Application/SlateUser.cpp
+++ b/Engine/Source/Runtime/Slate/Private/Framework/Application/SlateUser.cpp
@@ -575,6 +575,15 @@ void FSlateUser::CloseTooltip()
 	}
 }
 
+FVector2f FSlateUser::GetTooltipPosition() const
+{
+	if (TooltipWindowPtr.IsValid())
+	{
+		return TooltipWindowPtr.Pin()->GetPositionInScreen();
+	}
+	return FVector2f::Zero();
+}
+
 void FSlateUser::SetUserNavigationConfig(TSharedPtr<FNavigationConfig> InNavigationConfig)
 {
 	if (UserNavigationConfig)
diff --git a/Engine/Source/Runtime/Slate/Public/Framework/Application/SlateUser.h b/Engine/Source/Runtime/Slate/Public/Framework/Application/SlateUser.h
--- a/Engine/Source/Runtime/Slate/Public/Framework/Application/SlateUser.h
+++ b/Engine/Source/Runtime/Slate/Public/Framework/Application/SlateUser.h
@@ -118,6 +118,7 @@ public:
 	
 	SLATE_API void ShowTooltip(const TSharedRef<IToolTip>& InTooltip, const UE::Slate::FDeprecateVector2DParameter& InSpawnLocation);
 	SLATE_API void CloseTooltip();
+	SLATE_API FVector2f GetTooltipPosition() const;
 
 	const FGestureDetector& GetGestureDetector() const { return GestureDetector; }
 
