Prevent server crash on TranslationPicker command

When invoked via EngineGlobals.cpp `static FAutoConsoleCommand GTranslationPicker` the server would crash.

#rnx

diff --git a/Engine/Source/Developer/TranslationEditor/Private/TranslationPickerWidget.cpp b/Engine/Source/Developer/TranslationEditor/Private/TranslationPickerWidget.cpp
--- a/Engine/Source/Developer/TranslationEditor/Private/TranslationPickerWidget.cpp
+++ b/Engine/Source/Developer/TranslationEditor/Private/TranslationPickerWidget.cpp
@@ -133,6 +133,11 @@ bool TranslationPickerManager::bDrawBoxes = true;
 
 bool TranslationPickerManager::OpenPickerWindow()
 {
+	if (!FSlateApplicationBase::IsInitialized())
+	{
+		return false;
+	}
+
 	if (PickerWindow.IsValid() || PickerWindowWidget.IsValid())
 	{
 		return false;
diff --git a/Engine/Source/Runtime/SlateCore/Public/Application/SlateApplicationBase.h b/Engine/Source/Runtime/SlateCore/Public/Application/SlateApplicationBase.h
--- a/Engine/Source/Runtime/SlateCore/Public/Application/SlateApplicationBase.h
+++ b/Engine/Source/Runtime/SlateCore/Public/Application/SlateApplicationBase.h
@@ -529,6 +529,7 @@ public:
 	 */
 	static FSlateApplicationBase& Get( )
 	{
+		check(CurrentBaseApplication.IsValid());
 		checkSlow(IsThreadSafeForSlateRendering());
 		return *CurrentBaseApplication;
 	}
